NNXRuntimeHLSL variadic operators
Max 4 inputs version

#rnx

diff --git a/Engine/Plugins/Experimental/NNX/Source/NNXHlslShaders/Public/NNXElementWiseVariadicCS.h b/Engine/Plugins/Experimental/NNX/Source/NNXHlslShaders/Public/NNXElementWiseVariadicCS.h
--- a/Engine/Plugins/Experimental/NNX/Source/NNXHlslShaders/Public/NNXElementWiseVariadicCS.h
+++ b/Engine/Plugins/Experimental/NNX/Source/NNXHlslShaders/Public/NNXElementWiseVariadicCS.h
@@ -14,17 +14,19 @@ class NNXHLSLSHADERS_API FMLElementWiseVariadicCS : public FGlobalShader
 
 public:
 
-	NNXRT_ELEMENTWISEVARIADIC_PARAMETER_STRUCT()
-
 	static const uint32 THREADGROUP_SIZE_X;
+	static const uint32 MAX_NUM_INPUT = 4;
 
 	class FOperatorType : SHADER_PERMUTATION_ENUM_CLASS("OP_TYPENAME", EMLElementWiseVariadicOperatorType);
 	class FApplyScale : SHADER_PERMUTATION_BOOL("APPLYSCALE");
-	class FNumInput : SHADER_PERMUTATION_RANGE_INT("NUMINPUT", 1, 4);
-	using FPermutationDomain = TShaderPermutationDomain<FOperatorType, FApplyScale, FNumInput>;
+	class FOutputAsInput : SHADER_PERMUTATION_BOOL("OUTPUTASINPUT");
+	class FNumInput : SHADER_PERMUTATION_RANGE_INT("NUMINPUT", 1, MAX_NUM_INPUT);
+	using FPermutationDomain = TShaderPermutationDomain<FOperatorType, FApplyScale, FOutputAsInput, FNumInput>;
 
 	static void ModifyCompilationEnvironment(const FGlobalShaderPermutationParameters& InParameters, FShaderCompilerEnvironment& OutEnvironment);
 
+	NNXRT_ELEMENTWISEVARIADIC_PARAMETER_STRUCT()
+
 private:
 
 	static const FString GetOpFunc(EMLElementWiseVariadicOperatorType);
diff --git a/Engine/Plugins/Experimental/NNX/Source/NNXHlslShaders/Shaders/ElementWiseVariadicOp.usf b/Engine/Plugins/Experimental/NNX/Source/NNXHlslShaders/Shaders/ElementWiseVariadicOp.usf
--- a/Engine/Plugins/Experimental/NNX/Source/NNXHlslShaders/Shaders/ElementWiseVariadicOp.usf
+++ b/Engine/Plugins/Experimental/NNX/Source/NNXHlslShaders/Shaders/ElementWiseVariadicOp.usf
@@ -67,8 +67,11 @@ void ElementWiseVariadicOp(in const uint3 DispatchThreadID : SV_DispatchThreadID
 			Result  = ELEMENTWISE_OP(Result, Input3[Input3Index]);
 		#endif
 
-		// Must correspond to EVariadicScale defined in NNXElementWiseVariadicCS.h
-		#if APPLYSCALE == 1 
+		#if OUTPUTASINPUT
+			Result  = ELEMENTWISE_OP(Result, Output[Index]);
+		#endif
+
+		#if APPLYSCALE
 			Result  = (BUFFER_TYPE)(Result * Scale);
 		#endif
 
diff --git a/Engine/Plugins/Experimental/NNX/Source/NNXQA/Private/NNXQATestsOperatorElementWiseVariadic.cpp b/Engine/Plugins/Experimental/NNX/Source/NNXQA/Private/NNXQATestsOperatorElementWiseVariadic.cpp
--- a/Engine/Plugins/Experimental/NNX/Source/NNXQA/Private/NNXQATestsOperatorElementWiseVariadic.cpp
+++ b/Engine/Plugins/Experimental/NNX/Source/NNXQA/Private/NNXQATestsOperatorElementWiseVariadic.cpp
@@ -23,8 +23,8 @@ namespace Test
 		TArray<FString> AutomationExcludedRuntime;
 		EMLTensorDataType TensorType = EMLTensorDataType::Float;
 
-		//TODO Remove when HLSL support is added. Just to verify test are green using ORT for now.
-		AutomationExcludedRuntime.Emplace(TEXT("NNXRuntimeHlsl"));
+		//Variadic op not yet implemented on Dml
+		AutomationExcludedRuntime.Emplace(TEXT("NNXRuntimeDml"));
 
 		FString TestSuffix;
 		for (auto&& ShapeInput: ShapeInputs)
diff --git a/Engine/Plugins/Experimental/NNX/Source/NNXRuntimeRDG/Private/NNXRuntimeHLSL.cpp b/Engine/Plugins/Experimental/NNX/Source/NNXRuntimeRDG/Private/NNXRuntimeHLSL.cpp
--- a/Engine/Plugins/Experimental/NNX/Source/NNXRuntimeRDG/Private/NNXRuntimeHLSL.cpp
+++ b/Engine/Plugins/Experimental/NNX/Source/NNXRuntimeRDG/Private/NNXRuntimeHLSL.cpp
@@ -5,6 +5,7 @@
 #include "NNXRuntimeFormat.h"
 #include "NNXRuntimeHLSLElementWiseBinaryOps.h"
 #include "NNXRuntimeHLSLElementWiseUnaryOps.h"
+#include "NNXRuntimeHLSLElementWiseVariadicOps.h"
 #include "NNXRuntimeHLSLHelper.h"
 #include "NNXRuntimeRDG.h"
 
@@ -143,6 +144,7 @@ public:
 
 		RegisterElementWiseUnaryOperators(*registry);
 		RegisterElementWiseBinaryOperators(*registry);
+		RegisterElementWiseVariadicOperators(*registry);
 
 		return true;
 	}
diff --git a/Engine/Plugins/Experimental/NNX/Source/NNXRuntimeRDG/Private/NNXRuntimeHLSLElementWiseVariadicOps.cpp b/Engine/Plugins/Experimental/NNX/Source/NNXRuntimeRDG/Private/NNXRuntimeHLSLElementWiseVariadicOps.cpp
new file mode 100644
--- /dev/null
+++ b/Engine/Plugins/Experimental/NNX/Source/NNXRuntimeRDG/Private/NNXRuntimeHLSLElementWiseVariadicOps.cpp
@@ -0,0 +1,170 @@
+// Copyright Epic Games, Inc. All Rights Reserved.
+
+#include "NNXRuntimeHLSLElementWiseVariadicOps.h"
+#include "NNXElementWiseVariadicCS.h"
+#include "NNXRuntimeHLSLHelper.h"
+
+namespace NNX
+{
+	DECLARE_GPU_STAT_NAMED(VariadicElementWiseOperatorHlsl, TEXT("NNX VariadicElementWiseOperatorHlsl"));
+
+	void AddOneVariadicOpPass(FRDGBuilder& GraphBuilder, 
+		TArrayView<const FMLTensorBinding> InputBindings,
+		TArrayView<const FMLTensorDesc> InputDesc,
+		const FMLTensorBinding& OutputBinding,
+		const FMLTensorDesc& OutputDesc,
+		bool OutputAsInput,
+		EMLElementWiseVariadicOperatorType OpType,
+		float Scale)
+	{
+		static_assert(FMLElementWiseVariadicCS::MAX_NUM_INPUT == 4, "This algorithm need to be adapted to math the shader.");
+		check(InputBindings.Num() == InputDesc.Num());
+		check(InputBindings.Num() <= FMLElementWiseVariadicCS::MAX_NUM_INPUT);
+		check(InputBindings.Num() > 0);
+
+		// SRVs & UAV creations
+		FRDGBufferUAVRef OutputUAV = GraphBuilder.CreateUAV(FRDGBufferUAVDesc(OutputBinding.Buffer, PF_R32_FLOAT));
+		FRDGBufferSRVRef InputsSRV[FMLElementWiseVariadicCS::MAX_NUM_INPUT] = { nullptr, nullptr, nullptr, nullptr };
+
+		for (int32 i = 0; i < InputBindings.Num(); ++i)
+		{
+			InputsSRV[i] = GraphBuilder.CreateSRV(FRDGBufferSRVDesc(InputBindings[i].Buffer, PF_R32_FLOAT));
+		}
+
+		// Set parameters
+		FIntVector ThreadGroupCount = ComputeElementWiseThreadGroups(OutputDesc.Num(), FMLElementWiseVariadicCS::THREADGROUP_SIZE_X);
+		FMLElementWiseVariadicCS::FParameters* Params = GraphBuilder.AllocParameters<FMLElementWiseVariadicCS::FParameters>();
+
+		Params->Input0 = InputsSRV[0];
+		Params->Input1 = InputsSRV[1];
+		Params->Input2 = InputsSRV[2];
+		Params->Input3 = InputsSRV[3];
+		Params->Output = OutputUAV;
+		FillTensorStrideForBroadcastShaderParameters(InputDesc[0], OutputDesc.Dimension, Params->Input0Info0, Params->Input0Info1);
+		if (InputBindings.Num() >= 2)
+		{
+			FillTensorStrideForBroadcastShaderParameters(InputDesc[1], OutputDesc.Dimension, Params->Input1Info0, Params->Input1Info1);
+		}
+		if (InputBindings.Num() >= 3)
+		{
+			FillTensorStrideForBroadcastShaderParameters(InputDesc[2], OutputDesc.Dimension, Params->Input2Info0, Params->Input2Info1);
+		}
+		if (InputBindings.Num() >= 4)
+		{
+			FillTensorStrideForBroadcastShaderParameters(InputDesc[3], OutputDesc.Dimension, Params->Input3Info0, Params->Input3Info1);
+		}
+		FillTensorStrideShaderParameters(OutputDesc, Params->OutInfo0, Params->OutInfo1);
+		Params->OutRank = OutputDesc.Dimension;
+		Params->Num = OutputDesc.Num();
+		Params->ThreadCountX = ThreadGroupCount.X * FMLElementWiseVariadicCS::THREADGROUP_SIZE_X;
+		Params->Scale = Scale;
+
+		// Shader variation
+		FMLElementWiseVariadicCS::FPermutationDomain PermutationVector;
+
+		PermutationVector.Set<FMLElementWiseVariadicCS::FOperatorType>(OpType);
+		PermutationVector.Set<FMLElementWiseVariadicCS::FApplyScale>(Scale != 1.0f);
+		PermutationVector.Set<FMLElementWiseVariadicCS::FOutputAsInput>(OutputAsInput);
+		PermutationVector.Set<FMLElementWiseVariadicCS::FNumInput>(InputBindings.Num());
+
+		// Add the pass to RDG
+		TShaderMapRef<FMLElementWiseVariadicCS> ComputeShader(GetGlobalShaderMap(GMaxRHIFeatureLevel), PermutationVector);
+		FComputeShaderUtils::AddPass(
+			GraphBuilder,
+			RDG_EVENT_NAME("FMLVariadicElementWiseOperatorHlsl_Dispatch"),
+			ERDGPassFlags::Compute | ERDGPassFlags::NeverCull,
+			ComputeShader,
+			Params,
+			ThreadGroupCount);
+	}
+
+/**
+ * Variadic Element-wise ML operator implementation
+ */
+template<EMLElementWiseVariadicOperatorType OpType>
+class FMLOperatorHlslElementWiseVariadic : public FMLOperatorHlsl
+{
+public:
+
+	static FMLOperatorHlsl* Create()
+	{
+		return new FMLOperatorHlslElementWiseVariadic();
+	}
+
+	virtual ~FMLOperatorHlslElementWiseVariadic() = default;
+
+private:
+
+	FMLOperatorHlslElementWiseVariadic() {}
+
+	TArray<FMLTensorDesc> InputDescs;
+	FMLTensorDesc OutputDesc;
+
+public:
+
+	virtual bool Initialize(TArrayView<const FMLTensorDesc> InputTensors, TArrayView<const FMLTensorDesc> OutputTensors) override
+	{
+		check(InputTensors.Num() > 0);
+		check(OutputTensors.Num() == 1);
+		
+		InputDescs.Append(InputTensors);
+		OutputDesc = OutputTensors[0];
+
+		return true;
+	}
+
+	virtual void Dispatch(FRDGBuilder& GraphBuilder, TArrayView<const FMLTensorBinding> InInputBindings, TArrayView<const FMLTensorBinding> OutOutputBindings) override
+	{
+		check(OutOutputBindings.Num() == 1);
+		check(InInputBindings.Num() == InputDescs.Num());
+
+		RDG_EVENT_SCOPE(GraphBuilder, "VariadicElementWiseOperatorHlsl");
+		RDG_GPU_STAT_SCOPE(GraphBuilder, VariadicElementWiseOperatorHlsl);
+
+		FMLTensorBinding PassInputBindings[FMLElementWiseVariadicCS::MAX_NUM_INPUT];
+		FMLTensorDesc PassInputDescs[FMLElementWiseVariadicCS::MAX_NUM_INPUT];
+		for (int32 InputOffset = 0; InputOffset < InInputBindings.Num(); InputOffset += FMLElementWiseVariadicCS::MAX_NUM_INPUT)
+		{
+			uint32 NumInputLeftToHandle = InInputBindings.Num() - InputOffset;
+			uint32 NumInputForPass = FMath::Min(FMLElementWiseVariadicCS::MAX_NUM_INPUT, NumInputLeftToHandle);
+			bool bIsFirstPass = (InputOffset == 0);
+			bool bIsLastPass = (NumInputLeftToHandle <= FMLElementWiseVariadicCS::MAX_NUM_INPUT);
+			float Scale = 1.0f;
+
+			for (uint32 i = 0; i < NumInputForPass; ++i)
+			{
+				PassInputBindings[i] = InInputBindings[InputOffset + i];
+				PassInputDescs[i] = InputDescs[InputOffset + i];
+			}
+
+			if (OpType == EMLElementWiseVariadicOperatorType::Mean && bIsLastPass)
+			{
+				Scale = 1.0f / InInputBindings.Num();
+			}
+			
+			AddOneVariadicOpPass(GraphBuilder,
+				MakeArrayView(PassInputBindings, NumInputForPass),
+				MakeArrayView(PassInputDescs, NumInputForPass),
+				OutOutputBindings[0],
+				OutputDesc,
+				!bIsFirstPass,
+				OpType, Scale);
+		}
+	}
+};
+
+bool RegisterElementWiseVariadicOperators(FMLOperatorRegistryHlsl& Registry)
+{
+	#define OP(Name) Registry.OpAdd(TEXT(#Name), FMLOperatorHlslElementWiseVariadic<EMLElementWiseVariadicOperatorType::Name>::Create)
+
+	OP(Max);
+	OP(Min);
+	OP(Mean);
+	OP(Sum);
+
+	#undef OP
+
+	return true;
+}
+
+} // NNX
diff --git a/Engine/Plugins/Experimental/NNX/Source/NNXRuntimeRDG/Private/NNXRuntimeHLSLElementWiseVariadicOps.h b/Engine/Plugins/Experimental/NNX/Source/NNXRuntimeRDG/Private/NNXRuntimeHLSLElementWiseVariadicOps.h
new file mode 100644
--- /dev/null
+++ b/Engine/Plugins/Experimental/NNX/Source/NNXRuntimeRDG/Private/NNXRuntimeHLSLElementWiseVariadicOps.h
@@ -0,0 +1,11 @@
+// Copyright Epic Games, Inc. All Rights Reserved.
+
+#pragma once
+
+#include "NNXRuntimeRDG.h"
+#include "NNXRuntimeHLSLOp.h"
+
+namespace NNX
+{
+	bool RegisterElementWiseVariadicOperators(FMLOperatorRegistryHlsl& Registry);
+} // NNX
