Fixed runtime error when running QA model tests on RDG runtimes

#rnx

diff --git a/Engine/Plugins/Experimental/NNX/Source/NNXRuntimeRDG/Private/NNXRuntimeDML.cpp b/Engine/Plugins/Experimental/NNX/Source/NNXRuntimeRDG/Private/NNXRuntimeDML.cpp
--- a/Engine/Plugins/Experimental/NNX/Source/NNXRuntimeRDG/Private/NNXRuntimeDML.cpp
+++ b/Engine/Plugins/Experimental/NNX/Source/NNXRuntimeRDG/Private/NNXRuntimeDML.cpp
@@ -1022,6 +1022,12 @@ bool FMLInferenceModelDml::Init(UMLInferenceModel* InModel, FDeviceContextDml* I
 		return false;
 	}
 
+	if (Format.Operators.Num() > 1)
+	{
+		UE_LOG(LogNNX, Warning, TEXT("Failed to create inference model, currently on single layer models are supported"));
+		return false;
+	}
+
 	DevCtx = InDevCtx;
 
 	// Loop over all operators in the model and create them
diff --git a/Engine/Plugins/Experimental/NNX/Source/NNXRuntimeRDG/Private/NNXRuntimeHLSL.cpp b/Engine/Plugins/Experimental/NNX/Source/NNXRuntimeRDG/Private/NNXRuntimeHLSL.cpp
--- a/Engine/Plugins/Experimental/NNX/Source/NNXRuntimeRDG/Private/NNXRuntimeHLSL.cpp
+++ b/Engine/Plugins/Experimental/NNX/Source/NNXRuntimeRDG/Private/NNXRuntimeHLSL.cpp
@@ -65,6 +65,12 @@ bool FMLInferenceModelHlsl::Init(UMLInferenceModel* InModel)
 		return false;
 	}
 
+	if (Format.Operators.Num() > 1)
+	{
+		UE_LOG(LogNNX, Warning, TEXT("Failed to create inference model, currently on single layer models are supported"));
+		return false;
+	}
+
 	// Loop over all operators in the model and create them
 	for (int32 Idx = 0; Idx < Format.Operators.Num(); ++Idx)
 	{
