NeuralNetworkInference plugin:
Fixed NNX Utils, Added Protobuf as private dependency to allow the use of the custom protobuf in consoles

diff --git a/Engine/Plugins/Experimental/NNX/Source/NNXUtils/NNXUtils.Build.cs b/Engine/Plugins/Experimental/NNX/Source/NNXUtils/NNXUtils.Build.cs
--- a/Engine/Plugins/Experimental/NNX/Source/NNXUtils/NNXUtils.Build.cs
+++ b/Engine/Plugins/Experimental/NNX/Source/NNXUtils/NNXUtils.Build.cs
@@ -1,34 +1,49 @@
 // Copyright Epic Games, Inc. All Rights Reserved.
 
-namespace UnrealBuildTool.Rules
+
+using UnrealBuildTool;
+using System.IO;
+
+public class NNXUtils : ModuleRules
 {
-    public class NNXUtils : ModuleRules
-    {
-        public NNXUtils(ReadOnlyTargetRules Target) : base(Target)
-        {
-            PublicDependencyModuleNames.AddRange(
-				new string[]
-				{
-					"NNXCore"
+	public NNXUtils(ReadOnlyTargetRules Target) : base(Target)
+	{
+		ShortName = "NNXUtils"; // Shorten to avoid path-too-long errors
+		PCHUsage = ModuleRules.PCHUsageMode.UseExplicitOrSharedPCHs;
+
+		PublicDependencyModuleNames.AddRange(
+			new string[]
+			{
+				"NNXCore"
+			}
+		);
+
+		PrivateDependencyModuleNames.AddRange(
+			new string[]
+			{
+				"Core",
+				"CoreUObject",
+				"Engine",
+				"NNX_ONNXRuntime",
+				"NNX_ONNX_1_11_0",
+				"NNX_ONNXRuntimeProto_1_11_0",
+				"ORTHelper",
 				}
 			);
 
-            PrivateDependencyModuleNames.AddRange(
-                new string[]
-                {
-					"Core",
-					"CoreUObject",
-					"Engine",
-
-					"NNX_ONNXRuntime",
-					"NNX_ONNX_1_11_0",
-					"NNX_ONNXRuntimeProto_1_11_0",
-					"ORTHelper",
-					"Protobuf"
+
+		if (Target.Platform == UnrealTargetPlatform.Win64 || 
+			Target.Platform == UnrealTargetPlatform.Linux ||
+			Target.Platform == UnrealTargetPlatform.Mac)
+		{
+			PrivateDependencyModuleNames.AddRange
+				(
+				new string[] {
+					"Protobuf",
+					"Re2" // ONNXRuntimeRE2
 				}
 			);
-
-			//PrivateDefinitions.Add("ONNX_ML");
+		}
 		}
     }
-}
+
