Fix non-unity compilation

diff --git a/Engine/Plugins/Experimental/NNX/Source/NNXQA/Private/NNXQAModelTest.cpp b/Engine/Plugins/Experimental/NNX/Source/NNXQA/Private/NNXQAModelTest.cpp
--- a/Engine/Plugins/Experimental/NNX/Source/NNXQA/Private/NNXQAModelTest.cpp
+++ b/Engine/Plugins/Experimental/NNX/Source/NNXQA/Private/NNXQAModelTest.cpp
@@ -8,6 +8,10 @@
 #include "Misc/Paths.h"
 #include "HAL/ConsoleManager.h"
 
+#if WITH_DEV_AUTOMATION_TESTS
+#include "Misc/AutomationTest.h"
+#endif
+
 namespace NNX 
 {
 namespace Test 
@@ -172,8 +176,6 @@ namespace Test
 
 #if WITH_DEV_AUTOMATION_TESTS
 
-	#include "Misc/AutomationTest.h"
-
 	IMPLEMENT_COMPLEX_AUTOMATION_TEST(FNNXModelTest, "System.Engine.MachineLearning.NNX.ModelTest",
 		EAutomationTestFlags::ApplicationContextMask | EAutomationTestFlags::FeatureMask | EAutomationTestFlags::EngineFilter)
 
diff --git a/Engine/Plugins/Experimental/NNX/Source/NNXQA/Private/NNXQAOperatorTest.cpp b/Engine/Plugins/Experimental/NNX/Source/NNXQA/Private/NNXQAOperatorTest.cpp
--- a/Engine/Plugins/Experimental/NNX/Source/NNXQA/Private/NNXQAOperatorTest.cpp
+++ b/Engine/Plugins/Experimental/NNX/Source/NNXQA/Private/NNXQAOperatorTest.cpp
@@ -10,6 +10,10 @@
 
 #include "HAL/ConsoleManager.h"
 
+#if WITH_DEV_AUTOMATION_TESTS
+#include "Misc/AutomationTest.h"
+#endif
+
 namespace NNX 
 {
 namespace Test 
@@ -114,8 +118,6 @@ namespace Test
 
 #if WITH_DEV_AUTOMATION_TESTS
 
-	#include "Misc/AutomationTest.h"
-
 	// Unary element wise operator category
 	IMPLEMENT_COMPLEX_AUTOMATION_TEST(FNNXOperatorTestElementWiseUnary, "System.Engine.MachineLearning.NNX.OperatorTest.UnaryElementWise",
 		EAutomationTestFlags::ApplicationContextMask | EAutomationTestFlags::FeatureMask | EAutomationTestFlags::EngineFilter)
@@ -156,7 +158,7 @@ namespace Test
 	IMPLEMENT_COMPLEX_AUTOMATION_TEST(FNNXOperatorTestElementWiseVariadic, "System.Engine.MachineLearning.NNX.OperatorTest.VariadicElementWise",
 		EAutomationTestFlags::ApplicationContextMask | EAutomationTestFlags::FeatureMask | EAutomationTestFlags::EngineFilter)
 
-		void FNNXOperatorTestElementWiseVariadic::GetTests(TArray<FString>& OutBeautifiedNames, TArray <FString>& OutTestCommands) const
+	void FNNXOperatorTestElementWiseVariadic::GetTests(TArray<FString>& OutBeautifiedNames, TArray <FString>& OutTestCommands) const
 	{
 		for (FTests::FTestSetup TestSetup : TestsOperatorElementWiseVariadic.TestSetups)
 		{
diff --git a/Engine/Plugins/Experimental/NNX/Source/NNXQA/Private/NNXQAUtils.cpp b/Engine/Plugins/Experimental/NNX/Source/NNXQA/Private/NNXQAUtils.cpp
--- a/Engine/Plugins/Experimental/NNX/Source/NNXQA/Private/NNXQAUtils.cpp
+++ b/Engine/Plugins/Experimental/NNX/Source/NNXQA/Private/NNXQAUtils.cpp
@@ -7,6 +7,8 @@
 #include "HAL/UnrealMemory.h"
 #include "Kismet/GameplayStatics.h"
 
+#include <functional>
+
 namespace NNX 
 {
 namespace Test 
diff --git a/Engine/Plugins/Experimental/NNX/Source/NNXRuntimeRDG/Private/NNXRuntimeDML.cpp b/Engine/Plugins/Experimental/NNX/Source/NNXRuntimeRDG/Private/NNXRuntimeDML.cpp
--- a/Engine/Plugins/Experimental/NNX/Source/NNXRuntimeRDG/Private/NNXRuntimeDML.cpp
+++ b/Engine/Plugins/Experimental/NNX/Source/NNXRuntimeRDG/Private/NNXRuntimeDML.cpp
@@ -5,6 +5,8 @@
 #include "NNXRuntimeFormat.h"
 #include "NNXOperator.h"
 
+#include "HAL/FileManager.h"
+
 #include "RenderGraphBuilder.h"
 #include "RenderGraphUtils.h"
 
