Fixed errors reported by static analysis

#rnx

diff --git a/Engine/Plugins/Experimental/NNX/Source/NNXRuntimeORT/Private/NNXRuntimeORT.cpp b/Engine/Plugins/Experimental/NNX/Source/NNXRuntimeORT/Private/NNXRuntimeORT.cpp
--- a/Engine/Plugins/Experimental/NNX/Source/NNXRuntimeORT/Private/NNXRuntimeORT.cpp
+++ b/Engine/Plugins/Experimental/NNX/Source/NNXRuntimeORT/Private/NNXRuntimeORT.cpp
@@ -385,12 +385,25 @@ bool FMLInferenceModelORTDml::InitializedAndConfigureMembers()
 	if (GDynamicRHI && GDynamicRHI->GetInterfaceType() == ERHIInterfaceType::D3D12)
 	{
 		RHI = static_cast<ID3D12DynamicRHI*>(GDynamicRHI);
-	}
 
-	if (!RHI)
+		if (!RHI)
+		{
+			UE_LOG(LogNNX, Warning, TEXT("Error:%s RHI is not supported by DirectML"), GDynamicRHI->GetName());
+			return false;
+		}
+	}
+	else
 	{
-		UE_LOG(LogNNX, Warning, TEXT("Failed to use D3D12 RHI, got %s RHI"), GDynamicRHI ? TEXT("null") : GDynamicRHI->GetName());
-		return false;
+		if (GDynamicRHI)
+		{
+			UE_LOG(LogNNX, Warning, TEXT("Error:%s RHI is not supported by DirectML"), GDynamicRHI->GetName());
+			return false;
+		}
+		else
+		{
+			UE_LOG(LogNNX, Warning, TEXT("Error:No RHI found"));
+			return false;
+		}
 	}
 
 	int DeviceIndex = 0;
diff --git a/Engine/Plugins/Experimental/NNX/Source/NNXRuntimeRDG/Private/NNXRuntimeDML.cpp b/Engine/Plugins/Experimental/NNX/Source/NNXRuntimeRDG/Private/NNXRuntimeDML.cpp
--- a/Engine/Plugins/Experimental/NNX/Source/NNXRuntimeRDG/Private/NNXRuntimeDML.cpp
+++ b/Engine/Plugins/Experimental/NNX/Source/NNXRuntimeRDG/Private/NNXRuntimeDML.cpp
@@ -832,9 +832,6 @@ FMLRuntimeDml::~FMLRuntimeDml()
 //
 bool FMLRuntimeDml::Init()
 {
-	RegisterElementWiseUnaryOperators();
-	RegisterElementWiseBinaryOperators();
-
 	HRESULT Res;
 
 	// In order to use DirectML we need D3D12
@@ -843,14 +840,30 @@ bool FMLRuntimeDml::Init()
 	if (GDynamicRHI && GDynamicRHI->GetInterfaceType() == ERHIInterfaceType::D3D12)
 	{
 		RHI = static_cast<ID3D12DynamicRHI*>(GDynamicRHI);
+
+		if (!RHI)
+		{
+			UE_LOG(LogNNX, Warning, TEXT("Error:%s RHI is not supported by DirectML"), GDynamicRHI->GetName());
+			return false;
+		}
 	}
-		
-	if (!RHI)
+	else
 	{
-		UE_LOG(LogNNX, Warning, TEXT("Failed to use D3D12 RHI, got %s RHI"), GDynamicRHI ? TEXT("null") : GDynamicRHI->GetName());
-		return false;
+		if (GDynamicRHI)
+		{
+			UE_LOG(LogNNX, Warning, TEXT("Error:%s RHI is not supported by DirectML"), GDynamicRHI->GetName());
+			return false;
+		}
+		else
+		{
+			UE_LOG(LogNNX, Warning, TEXT("Error:No RHI found"));
+			return false;
+		}
 	}
 
+	RegisterElementWiseUnaryOperators();
+	RegisterElementWiseBinaryOperators();
+
 	Ctx.DeviceIndex = 0;
 	Ctx.D3D12Device = RHI->RHIGetDevice(Ctx.DeviceIndex);
 
