Updated HLSL shader for GEMM operator to support stacks of matrices (used by MatMul operator).

#rnx

diff --git a/Engine/Plugins/Experimental/NNX/Source/NNXHlslShaders/Private/NNXGemmCS.cpp b/Engine/Plugins/Experimental/NNX/Source/NNXHlslShaders/Private/NNXGemmCS.cpp
--- a/Engine/Plugins/Experimental/NNX/Source/NNXHlslShaders/Private/NNXGemmCS.cpp
+++ b/Engine/Plugins/Experimental/NNX/Source/NNXHlslShaders/Private/NNXGemmCS.cpp
@@ -2,52 +2,192 @@
 
 #include "NNXGemmCS.h"
 #include "NNXCore.h"
+#include "Algo/Accumulate.h"
 
-EGemmAlgorithm FMLGemmCS::GetAlgorithm(const FMLGemmCS::FParameters& Parameters)
+namespace
 {
-	return EGemmAlgorithm::MultiWrite1x32;
+	struct GemmMatrixParameters
+	{
+		uint32 M = 0;
+		uint32 N = 0;
+		uint32 K = 0;
+		TArray<uint32> StackShapeA;
+		TArray<uint32> StackShapeB;
+		TArray<uint32> StackStrideA;
+		TArray<uint32> StackStrideB;
+	};
+
+	static GemmMatrixParameters GetMatrixParameters(TArray<uint32> ShapeA, TArray<uint32> ShapeB)
+	{
+		check(!ShapeA.IsEmpty());
+		check(!ShapeB.IsEmpty());
+
+		GemmMatrixParameters Result;
+		int32 NumStackDimensionsA = FMath::Max(ShapeA.Num() - 2, 0);
+		int32 NumStackDimensionsB = FMath::Max(ShapeB.Num() - 2, 0);
+		int32 NumStackDimensions = FMath::Max(NumStackDimensionsA, NumStackDimensionsB);
+
+		check(NumStackDimensions <= NNXRT_GEMM_MAX_NUM_STACK_DIMENSIONS);
+
+		// Check matrix stack dimensions
+		if (NumStackDimensionsA > 0 && NumStackDimensionsB > 0)
+		{
+			const int32 NumDimToCheck = FMath::Min(NumStackDimensionsA, NumStackDimensionsB);
+			for (int32 i = 0; i < NumDimToCheck; i++)
+			{
+				const uint32 VolumeA = ShapeA[NumStackDimensionsA - 1 - i];
+				const uint32 VolumeB = ShapeB[NumStackDimensionsB - 1 - i];
+
+				check(VolumeA == 1 || VolumeB == 1 || VolumeA == VolumeB);
+			}
+		}
+
+		const bool IsVectorA = ShapeA.Num() == 1;
+		const bool IsVectorB = ShapeB.Num() == 1;
+
+		Result.M = IsVectorA ? 1 : ShapeA[ShapeA.Num() - 2];
+		Result.N = IsVectorB ? 1 : ShapeB[ShapeB.Num() - 1];
+		Result.K = ShapeA[IsVectorA ? 0 : ShapeA.Num() - 1];
+		check(ShapeB[IsVectorB ? 0 : ShapeB.Num() - 2] == Result.K);
+
+		Result.StackShapeA.Init(1, NumStackDimensions);
+		Result.StackShapeB.Init(1, NumStackDimensions);
+
+		for (int32 i = 0; i < NumStackDimensions; i++)
+		{
+			Result.StackShapeA[Result.StackShapeA.Num() - 1 - i] = i < NumStackDimensionsA ? ShapeA[ShapeA.Num() - 3 - i] : 1;
+			Result.StackShapeB[Result.StackShapeB.Num() - 1 - i] = i < NumStackDimensionsB ? ShapeB[ShapeB.Num() - 3 - i] : 1;
+		}
+
+		Result.StackStrideA.Init(1, NumStackDimensions);
+		Result.StackStrideB.Init(1, NumStackDimensions);
+
+		for (int32 i = NumStackDimensions - 2; i >= 0; i--)
+		{
+			Result.StackStrideA[i] = Result.StackStrideA[i + 1] * Result.StackShapeA[i + 1];
+			Result.StackStrideB[i] = Result.StackStrideB[i + 1] * Result.StackShapeB[i + 1];
+		}
+
+		return Result;
+	}
+} // namespace
+
+void FMLGemmCS::ModifyCompilationEnvironment(const FGlobalShaderPermutationParameters& InParameters, FShaderCompilerEnvironment& OutEnvironment)
+{
+	FGlobalShader::ModifyCompilationEnvironment(InParameters, OutEnvironment);
+	OutEnvironment.SetDefine(TEXT("MAX_NUM_STACK_DIMENSIONS"), NNXRT_GEMM_MAX_NUM_STACK_DIMENSIONS);
+}
+
+void FMLGemmCS::FillInParameters(float Alpha, float Beta, int32 TransA, int32 TransB, uint32 M, uint32 N, uint32 K, uint32 CWidth, uint32 CHeight, float CScalar, FMLGemmCS::FParameters& Parameters)
+{
+	Parameters.Alpha = Alpha;
+	Parameters.Beta = Beta;
+	Parameters.TransA = TransA;
+	Parameters.TransB = TransB;
+	Parameters.M = M;
+	Parameters.N = N;
+	Parameters.K = K;
+	Parameters.MxK = M * K;
+	Parameters.KxN = K * N;
+	Parameters.MxN = M * N;
+	Parameters.CWidth = CWidth;
+	Parameters.CHeight = CHeight;
+	Parameters.CScalar = CScalar;
+}
+
+uint32 FMLGemmCS::GetShapeSize(TArray<uint32> Shape)
+{
+	auto Product = [](uint32 A, uint32 B) { return A * B; };
+
+	return Algo::Accumulate(Shape, 1, Product);
+}
+
+uint32 FMLGemmCS::GetMatMulOutputSize(TArray<uint32> ShapeA, TArray<uint32> ShapeB)
+{
+	GemmMatrixParameters Params = GetMatrixParameters(ShapeA, ShapeB);
+
+	return FMath::Max(GetShapeSize(Params.StackShapeA), GetShapeSize(Params.StackShapeB)) * Params.M * Params.N;
 }
 
-FIntVector FMLGemmCS::GetGroupCount(const FMLGemmCS::FParameters& Parameters, EGemmAlgorithm Algorithm)
+void FMLGemmCS::FillInParametersMatMul(TArray<uint32> ShapeA, TArray<uint32> ShapeB, FMLGemmCS::FParameters& Parameters)
 {
-	uint32 ThreadGroupCountValueY = FMath::DivideAndRoundUp((uint32)Parameters.M, (uint32)8);
-	uint32 ThreadGroupCountValueX = FMath::DivideAndRoundUp((uint32)Parameters.N, (uint32)8);
+	GemmMatrixParameters Params = GetMatrixParameters(ShapeA, ShapeB);
+
+	Parameters.Alpha = 1.0f;
+	Parameters.Beta = 1.0f;
+	Parameters.TransA = 0;
+	Parameters.TransB = 0;
+	Parameters.M = Params.M;
+	Parameters.N = Params.N;
+	Parameters.K = Params.K;
+	Parameters.MxK = Params.M * Params.K;
+	Parameters.KxN = Params.K * Params.N;
+	Parameters.MxN = Params.M * Params.N;
+	Parameters.CWidth = 0;
+	Parameters.CHeight = 0;
+	Parameters.CScalar = 0;
+
+	for (int32 i = 0; i < Params.StackShapeA.Num(); i++)
+	{
+		Parameters.StackShapeA_StackShapeB_StackStrideA_StackStrideB[i] = FUint32Vector4(Params.StackShapeA[i], Params.StackShapeB[i], Params.StackStrideA[i], Params.StackStrideB[i]);
+	}
+}
+
+FIntVector FMLGemmCS::GetGroupCount(const FMLGemmCS::FParameters& Parameters, EGemmAlgorithm Algorithm, int32 NumStackDimensions)
+{
+	int32 OutputStackSize = 1;
+	for (int32 i = 0; i < NumStackDimensions; i++)
+	{
+		OutputStackSize *= FMath::Max(
+			Parameters.StackShapeA_StackShapeB_StackStrideA_StackStrideB[i].X,
+			Parameters.StackShapeA_StackShapeB_StackStrideA_StackStrideB[i].Y);
+	}
+
+	int32 ThreadGroupCountValueZ = FMath::DivideAndRoundUp(OutputStackSize, 1);
+	int32 ThreadGroupCountValueY = FMath::DivideAndRoundUp((int32)Parameters.M, 8);
+	int32 ThreadGroupCountValueX = FMath::DivideAndRoundUp((int32)Parameters.N, 8);
 	switch (Algorithm)
 	{
-		case EGemmAlgorithm::Simple8x8:
-		case EGemmAlgorithm::SharedMemory8x8:
-			ThreadGroupCountValueX = FMath::DivideAndRoundUp((uint32)Parameters.N, (uint32)8);
-			ThreadGroupCountValueY = FMath::DivideAndRoundUp((uint32)Parameters.M, (uint32)8);
-			break;
-		case EGemmAlgorithm::Simple16x16:
-		case EGemmAlgorithm::SharedMemory16x16:
-		case EGemmAlgorithm::MultiWrite1x16:
-		case EGemmAlgorithm::MultiWrite2x16:
-			ThreadGroupCountValueX = FMath::DivideAndRoundUp((uint32)Parameters.N, (uint32)16);
-			ThreadGroupCountValueY = FMath::DivideAndRoundUp((uint32)Parameters.M, (uint32)16);
-			break;
-		case EGemmAlgorithm::Simple32x32:
-		case EGemmAlgorithm::SharedMemory32x32:
-		case EGemmAlgorithm::MultiWrite1x32:
-		case EGemmAlgorithm::MultiWrite2x32:
-		case EGemmAlgorithm::MultiWrite4x32:
-			ThreadGroupCountValueX = FMath::DivideAndRoundUp((uint32)Parameters.N, (uint32)32);
-			ThreadGroupCountValueY = FMath::DivideAndRoundUp((uint32)Parameters.M, (uint32)32);
-			break;
-		case EGemmAlgorithm::Simple256x1:
-			ThreadGroupCountValueX = FMath::DivideAndRoundUp((uint32)Parameters.N, (uint32)256);
-			ThreadGroupCountValueY = (uint32)Parameters.M;
-			break;
-		case EGemmAlgorithm::MultiWrite2x64:
-		case EGemmAlgorithm::MultiWrite4x64:
-		case EGemmAlgorithm::MultiWrite8x64:
-			ThreadGroupCountValueX = FMath::DivideAndRoundUp((uint32)Parameters.N, (uint32)64);
-			ThreadGroupCountValueY = FMath::DivideAndRoundUp((uint32)Parameters.M, (uint32)64);
-			break;
-		default:
-			break;
+	case EGemmAlgorithm::Simple8x8:
+	case EGemmAlgorithm::SharedMemory8x8:
+		ThreadGroupCountValueX = FMath::DivideAndRoundUp((int32)Parameters.N, 8);
+		ThreadGroupCountValueY = FMath::DivideAndRoundUp((int32)Parameters.M, 8);
+		break;
+	case EGemmAlgorithm::Simple16x16:
+	case EGemmAlgorithm::SharedMemory16x16:
+	case EGemmAlgorithm::MultiWrite1x16:
+	case EGemmAlgorithm::MultiWrite2x16:
+		ThreadGroupCountValueX = FMath::DivideAndRoundUp((int32)Parameters.N, 16);
+		ThreadGroupCountValueY = FMath::DivideAndRoundUp((int32)Parameters.M, 16);
+		break;
+	case EGemmAlgorithm::Simple32x32:
+	case EGemmAlgorithm::SharedMemory32x32:
+	case EGemmAlgorithm::MultiWrite1x32:
+	case EGemmAlgorithm::MultiWrite2x32:
+	case EGemmAlgorithm::MultiWrite4x32:
+		ThreadGroupCountValueX = FMath::DivideAndRoundUp((int32)Parameters.N, 32);
+		ThreadGroupCountValueY = FMath::DivideAndRoundUp((int32)Parameters.M, 32);
+		break;
+	case EGemmAlgorithm::Simple256x1:
+		ThreadGroupCountValueX = FMath::DivideAndRoundUp((int32)Parameters.N, 256);
+		ThreadGroupCountValueY = Parameters.M;
+		break;
+	case EGemmAlgorithm::MultiWrite2x64:
+	case EGemmAlgorithm::MultiWrite4x64:
+	case EGemmAlgorithm::MultiWrite8x64:
+		ThreadGroupCountValueX = FMath::DivideAndRoundUp((int32)Parameters.N, 64);
+		ThreadGroupCountValueY = FMath::DivideAndRoundUp((int32)Parameters.M, 64);
+		break;
+	default:
+		break;
 	}
-	return FIntVector(ThreadGroupCountValueX, ThreadGroupCountValueY, 1);
+
+	return FIntVector(ThreadGroupCountValueX, ThreadGroupCountValueY, ThreadGroupCountValueZ);
+}
+
+EGemmAlgorithm FMLGemmCS::GetAlgorithm(const FMLGemmCS::FParameters& Parameters)
+{
+	return EGemmAlgorithm::MultiWrite1x32;
 }
 
 IMPLEMENT_GLOBAL_SHADER(FMLGemmCS, "/NNX/GemmOp.usf", "main", SF_Compute);
\ No newline at end of file
diff --git a/Engine/Plugins/Experimental/NNX/Source/NNXHlslShaders/Public/NNXGemmCS.h b/Engine/Plugins/Experimental/NNX/Source/NNXHlslShaders/Public/NNXGemmCS.h
--- a/Engine/Plugins/Experimental/NNX/Source/NNXHlslShaders/Public/NNXGemmCS.h
+++ b/Engine/Plugins/Experimental/NNX/Source/NNXHlslShaders/Public/NNXGemmCS.h
@@ -40,11 +40,21 @@ class NNXHLSLSHADERS_API FMLGemmCS : public FGlobalShader
 
 	class FGemmCScalar : SHADER_PERMUTATION_ENUM_CLASS("C_SCALAR", EGemmCScalar);
 	class FGemmAlgorithm : SHADER_PERMUTATION_ENUM_CLASS("ALGORITHM", EGemmAlgorithm);
-	using FPermutationDomain = TShaderPermutationDomain<FGemmCScalar, FGemmAlgorithm>;
+	class FGemmNumStackDimensions : SHADER_PERMUTATION_RANGE_INT("NUM_STACK_DIMENSIONS", 0, NNXRT_GEMM_MAX_NUM_STACK_DIMENSIONS);
+	using FPermutationDomain = TShaderPermutationDomain<FGemmCScalar, FGemmAlgorithm, FGemmNumStackDimensions>;
 
 public:
 	NNXRT_GEMM_PARAMETER_STRUCT()
 
-	static FIntVector GetGroupCount(const FMLGemmCS::FParameters& Parameters, EGemmAlgorithm Algorithm);
+	static void ModifyCompilationEnvironment(const FGlobalShaderPermutationParameters& InParameters, FShaderCompilerEnvironment& OutEnvironment);
+
+	static void FillInParameters(float Alpha, float Beta, int32 TransA, int32 TransB, uint32 M, uint32 N, uint32 K, uint32 CWidth, uint32 CHeight, float CScalar, FMLGemmCS::FParameters& Parameters);
+
+	static uint32 GetShapeSize(TArray<uint32> Shape);
+	static uint32 GetMatMulOutputSize(TArray<uint32> ShapeA, TArray<uint32> ShapeB);
+
+	static void FillInParametersMatMul(TArray<uint32> ShapeA, TArray<uint32> ShapeB, FMLGemmCS::FParameters& Parameters);
+
+	static FIntVector GetGroupCount(const FMLGemmCS::FParameters& Parameters, EGemmAlgorithm Algorithm, int32 NumStackDimensions);
 	static EGemmAlgorithm GetAlgorithm(const FMLGemmCS::FParameters& Parameters);
 };
diff --git a/Engine/Plugins/Experimental/NNX/Source/NNXHlslShaders/Public/NNXShaderParameters.h b/Engine/Plugins/Experimental/NNX/Source/NNXHlslShaders/Public/NNXShaderParameters.h
--- a/Engine/Plugins/Experimental/NNX/Source/NNXHlslShaders/Public/NNXShaderParameters.h
+++ b/Engine/Plugins/Experimental/NNX/Source/NNXHlslShaders/Public/NNXShaderParameters.h
@@ -88,6 +88,8 @@ NNXRT_ELEMENTWISEBINARY_PARAMETER_STRUCT()
 #endif
 
 //Gemm
+#define NNXRT_GEMM_MAX_NUM_STACK_DIMENSIONS 8
+
 #define NNXRT_GEMM_PARAMETER_STRUCT() \
 	NNXRT_BEGIN_SHADER_PARAMETER_STRUCT(FMLGemmParameters) \
 		SHADER_PARAMETER(float, Alpha) \
@@ -97,9 +99,13 @@ NNXRT_ELEMENTWISEBINARY_PARAMETER_STRUCT()
 		SHADER_PARAMETER(uint32, M) \
 		SHADER_PARAMETER(uint32, N) \
 		SHADER_PARAMETER(uint32, K) \
+		SHADER_PARAMETER(uint32, MxK) \
+		SHADER_PARAMETER(uint32, KxN) \
+		SHADER_PARAMETER(uint32, MxN) \
 		SHADER_PARAMETER(uint32, CWidth) \
 		SHADER_PARAMETER(uint32, CHeight) \
 		SHADER_PARAMETER(float, CScalar) \
+		SHADER_PARAMETER_ARRAY(FUint32Vector4, StackShapeA_StackShapeB_StackStrideA_StackStrideB, [NNXRT_GEMM_MAX_NUM_STACK_DIMENSIONS]) \
 		NNXRT_RDG_BUFFER_SRV(Buffer<float>, A) \
 		NNXRT_RDG_BUFFER_SRV(Buffer<float>, B) \
 		NNXRT_RDG_BUFFER_SRV(Buffer<float>, C) \
diff --git a/Engine/Plugins/Experimental/NNX/Source/NNXHlslShaders/Shaders/GemmOp.usf b/Engine/Plugins/Experimental/NNX/Source/NNXHlslShaders/Shaders/GemmOp.usf
--- a/Engine/Plugins/Experimental/NNX/Source/NNXHlslShaders/Shaders/GemmOp.usf
+++ b/Engine/Plugins/Experimental/NNX/Source/NNXHlslShaders/Shaders/GemmOp.usf
@@ -14,6 +14,9 @@ int TransB;
 uint M;
 uint N;
 uint K;
+uint MxK;
+uint KxN;
+uint MxN;
 uint CWidth;
 uint CHeight;
 
@@ -24,7 +27,53 @@ Buffer<BUFFER_TYPE> B;
 Buffer<BUFFER_TYPE> C;
 RWBuffer<BUFFER_TYPE> Y;
 
+// x: StackShapeA[0], ..., StackShapeA[NUM_STACK_DIMENSIONS - 1]
+// y: StackShapeB[0], ..., StackShapeB[NUM_STACK_DIMENSIONS - 1]
+// z: StackStrideA[0], ..., StackStrideA[NUM_STACK_DIMENSIONS - 1]
+// w: StackStrideB[0], ..., StackStrideB[NUM_STACK_DIMENSIONS - 1]
+uint4 StackShapeA_StackShapeB_StackStrideA_StackStrideB[MAX_NUM_STACK_DIMENSIONS];
 
+uint3 GetMatrixStackOffsets(const uint3 GroupID)
+{
+// Must correspond to FGemmNumDimensions defined in NNXGemmCS.h
+#if NUM_STACK_DIMENSIONS == 0
+	return uint3(0, 0, 0);
+#else
+	uint stackCoordA[NUM_STACK_DIMENSIONS];
+	uint stackCoordB[NUM_STACK_DIMENSIONS];
+
+	uint stackStrideOutput[NUM_STACK_DIMENSIONS];
+	stackStrideOutput[NUM_STACK_DIMENSIONS - 1] = 1;
+	for (int i = NUM_STACK_DIMENSIONS - 2; i >= 0; i--)
+	{
+		uint maxDim = max(StackShapeA_StackShapeB_StackStrideA_StackStrideB[i + 1][0], StackShapeA_StackShapeB_StackStrideA_StackStrideB[i + 1][1]);
+		stackStrideOutput[i] = maxDim * stackStrideOutput[i + 1];
+	}
+
+	uint idx = GroupID.z;
+	for (int i = 0; i < NUM_STACK_DIMENSIONS; i++)
+	{
+		uint coord = idx / stackStrideOutput[i];
+		stackCoordA[i] = coord % StackShapeA_StackShapeB_StackStrideA_StackStrideB[i][0];
+		stackCoordB[i] = coord % StackShapeA_StackShapeB_StackStrideA_StackStrideB[i][1];
+
+		idx = idx % stackStrideOutput[i];
+	}
+
+	uint stackOffsetA = 0;
+	uint stackOffsetB = 0;
+	for (int i = 0; i < NUM_STACK_DIMENSIONS; i++)
+	{
+		stackOffsetA += stackCoordA[i] * StackShapeA_StackShapeB_StackStrideA_StackStrideB[i][2];
+		stackOffsetB += stackCoordB[i] * StackShapeA_StackShapeB_StackStrideA_StackStrideB[i][3];
+	}
+
+	return uint3(
+		stackOffsetA	* MxK,
+		stackOffsetB	* KxN,
+		GroupID.z		* MxN);
+#endif
+}
 
 WORK_TYPE GetBetaTimesC(const uint3 DispatchThreadID)
 {
@@ -50,7 +99,7 @@ WORK_TYPE GetBetaTimesC(const uint3 DispatchThreadID)
 #elif ALGORITHM == 3
 [numthreads(256, 1, 1)]
 #endif
-void main(in const uint3 DispatchThreadID : SV_DispatchThreadID)
+void main(in const uint3 DispatchThreadID : SV_DispatchThreadID, in const uint3 GroupID : SV_GroupID)
 {
 	if (DispatchThreadID.y >= M || DispatchThreadID.x >= N)
 	{
@@ -82,6 +131,10 @@ void main(in const uint3 DispatchThreadID : SV_DispatchThreadID)
 		BIdx = DispatchThreadID.x * K; // Index of the first element in the row
 		BIdxIncrement = 1; // The loop goes over the whole row
 	}
+
+	const uint3 StackOffsets = GetMatrixStackOffsets(GroupID);
+	AIdx += StackOffsets.x;
+	BIdx += StackOffsets.y;
 	
 	WORK_TYPE Result = 0.0;
 	for (int k = 0; k < K; k++)
@@ -91,7 +144,7 @@ void main(in const uint3 DispatchThreadID : SV_DispatchThreadID)
 		BIdx += BIdxIncrement;
 	}
 	
-	Y[DispatchThreadID.y * N + DispatchThreadID.x] = WRITE(Alpha * Result + GetBetaTimesC(DispatchThreadID));
+	Y[StackOffsets.z + DispatchThreadID.y * N + DispatchThreadID.x] = WRITE(Alpha * Result + GetBetaTimesC(DispatchThreadID));
 }
 
 #elif ALGORITHM < 7
@@ -108,7 +161,7 @@ groupshared WORK_TYPE SharedMemoryA[GROUP_SIZE * GROUP_SIZE];
 groupshared WORK_TYPE SharedMemoryB[GROUP_SIZE * GROUP_SIZE];
 
 [numthreads(GROUP_SIZE, GROUP_SIZE, 1)]
-void main(in const uint3 DispatchThreadID : SV_DispatchThreadID, in const uint3 GroupThreadID : SV_GroupThreadID)
+void main(in const uint3 DispatchThreadID : SV_DispatchThreadID, in const uint3 GroupThreadID : SV_GroupThreadID, in const uint3 GroupID : SV_GroupID)
 {
 	// Calculate the number of steps each group needs to make to cover the full length of the matrix
 	uint NumGroupSteps = (uint)ceil((float)K / (float)GROUP_SIZE);
@@ -143,6 +196,10 @@ void main(in const uint3 DispatchThreadID : SV_DispatchThreadID, in const uint3
 		BIdx = DispatchThreadID.x * K + GroupThreadID.y;
 		BIdxIncrement = GROUP_SIZE;
 	}
+
+	const uint3 StackOffsets = GetMatrixStackOffsets(GroupID);
+	AIdx += StackOffsets.x;
+	BIdx += StackOffsets.y;
 	
 	// Multiply and accumulate the blocks along one dimension of the matrices
 	WORK_TYPE Result = 0.0;
@@ -195,7 +252,7 @@ void main(in const uint3 DispatchThreadID : SV_DispatchThreadID, in const uint3
 		return;
 	}
 	
-	Y[DispatchThreadID.y * N + DispatchThreadID.x] = WRITE(Alpha * Result + GetBetaTimesC(DispatchThreadID));
+	Y[StackOffsets.z + DispatchThreadID.y * N + DispatchThreadID.x] = WRITE(Alpha * Result + GetBetaTimesC(DispatchThreadID));
 }
 
 #else
@@ -292,6 +349,10 @@ void main(in const uint3 DispatchThreadID : SV_DispatchThreadID, in const uint3
 		BIdxInnerIncrement = GROUP_SIZE_Y;
 	}
 
+	const uint3 StackOffsets = GetMatrixStackOffsets(GroupID);
+	AIdx += StackOffsets.x;
+	BIdx += StackOffsets.y;
+
 	// Multiply and accumulate the blocks
 	uint GroupIdx = 0;
 	for (uint GroupStep = 0; GroupStep < NumGroupSteps; GroupStep++) 
@@ -379,7 +440,7 @@ void main(in const uint3 DispatchThreadID : SV_DispatchThreadID, in const uint3
 			return;
 		}
 		
-		Y[RowIdx + Column] = WRITE(Alpha * Results[LoadIdx] + GetBetaTimesC(uint3(Column, Row, 0)));
+		Y[StackOffsets.z + RowIdx + Column] = WRITE(Alpha * Results[LoadIdx] + GetBetaTimesC(uint3(Column, Row, 0)));
     }
 }
 
