Implemented FMLInferenceModel::EnqueueRDG()

#rnx

diff --git a/Engine/Plugins/Experimental/NNX/Source/NNXRuntimeRDG/Private/NNXRuntimeRDG.cpp b/Engine/Plugins/Experimental/NNX/Source/NNXRuntimeRDG/Private/NNXRuntimeRDG.cpp
--- a/Engine/Plugins/Experimental/NNX/Source/NNXRuntimeRDG/Private/NNXRuntimeRDG.cpp
+++ b/Engine/Plugins/Experimental/NNX/Source/NNXRuntimeRDG/Private/NNXRuntimeRDG.cpp
@@ -77,65 +77,89 @@ bool FMLInferenceModelRDG::LoadModel(UMLInferenceModel* InModel, FMLRuntimeForma
 	return true;
 }
 
-//
-//
-//
+/**
+ * Run the inference model (synchronous version)
+ */
 int FMLInferenceModelRDG::Run(TArrayView<const FMLTensorBinding> InInputBindings, TArrayView<const FMLTensorBinding> OutOutputBindings)
 {
-	std::atomic<bool> bIsDone = false;
+	FEvent* Signal = FGenericPlatformProcess::GetSynchEventFromPool(false);
+	int		Res = 0;
 
 	ENQUEUE_RENDER_COMMAND(FMLInferenceModel_Run)
 	(
-		[&bIsDone, this, InInputBindings, &OutOutputBindings](FRHICommandListImmediate& RHICmdList)
+		[&Signal, &Res, this, InInputBindings, OutOutputBindings](FRHICommandListImmediate& RHICmdList)
 		{
 			FRDGBuilder	GraphBuilder(RHICmdList);
 
-			// Process input tensors, and if required, allocate RDG buffers
-			FMLTensorBindingArray	RDGInputBindings;
-			FMLIntArray				RDGUploadIndices;
-
-			SetTensors(GraphBuilder, RDGInputBindings, RDGUploadIndices, InInputBindings, InputTensors);
+			Res = EnqueueRDG(GraphBuilder, InInputBindings, OutOutputBindings);
+			if (Res == 0)
+			{
+				GraphBuilder.Execute();
+			}
+			
+			Signal->Trigger();
+		}
+	);
 
-			// Process output tensors, and if required, allocate RDG buffers
-			FMLTensorBindingArray	RDGOutputBindings;
-			FMLIntArray				RDGReadbackIndices;
+	// We need to wait for render thread to finish
+	Signal->Wait();
 
-			SetTensors(GraphBuilder, RDGOutputBindings, RDGReadbackIndices, OutOutputBindings, OutputTensors);
+	return Res;
+}
 
-			// If required, upload input tensors to GPU
-			AddTensorUploads_RenderThread(GraphBuilder, RDGUploadIndices, RDGInputBindings, InInputBindings);
+/**
+ * Enqueue operators to RDG, the caller will run the GraphBuilder.Execute()
+ */
+int FMLInferenceModelRDG::EnqueueRDG(FRDGBuilder& GraphBuilder, TArrayView<const FMLTensorBinding> InInputBindings, TArrayView<const FMLTensorBinding> InOutputBindings)
+{
+	check(IsInRenderingThread());
 
-			// We can now dispatch operators
-			AddDispatchOps_RenderThread(GraphBuilder, RDGInputBindings, RDGOutputBindings);
+	int Res;
 
-			// If required, readback the output tensors to CPU
-			AddTensorReadbacks_RenderThread(GraphBuilder, RDGReadbackIndices, RDGOutputBindings, OutOutputBindings);
+	// Process input tensors, and if required, allocate RDG buffers
+	FMLTensorBindingArray	RDGInputBindings;
+	FMLIntArray				RDGUploadIndices;
 
-			GraphBuilder.Execute();
+	Res = SetTensors(GraphBuilder, RDGInputBindings, RDGUploadIndices, InInputBindings, InputTensors);
+	if (Res != 0)
+	{
+		UE_LOG(LogNNX, Warning, TEXT("Invalid input tensor binding type for tensor index:%d"), Res);
+		return -1;
+	}
 
-			bIsDone = true;
-		}
-	);
+	// Process output tensors, and if required, allocate RDG buffers
+	FMLTensorBindingArray	RDGOutputBindings;
+	FMLIntArray				RDGReadbackIndices;
 
+	Res = SetTensors(GraphBuilder, RDGOutputBindings, RDGReadbackIndices, InOutputBindings, OutputTensors);
+	if (Res != 0)
+	{
+		UE_LOG(LogNNX, Warning, TEXT("Invalid output tensor binding type for tensor index:%d"), Res);
+		return -1;
+	}
 
-	// We need to wait for render thread to finish
-	while (!bIsDone)
+	// If required, upload input tensors to GPU
+	if (!RDGUploadIndices.IsEmpty())
 	{
-		FPlatformProcess::Sleep(0.1e-3);
+		AddTensorUploads_RenderThread(GraphBuilder, RDGUploadIndices, RDGInputBindings, InInputBindings);
 	}
 
-	return 0;
-}
+	// We can now dispatch operators
+	AddDispatchOps_RenderThread(GraphBuilder, RDGInputBindings, RDGOutputBindings);
+
+	// If required, readback the output tensors to CPU
+	if (!RDGReadbackIndices.IsEmpty())
+	{
+		AddTensorReadbacks_RenderThread(GraphBuilder, RDGReadbackIndices, RDGOutputBindings, InOutputBindings);
+	}
 
-//
-//
-//
-int FMLInferenceModelRDG::EnqueueRDG(FRDGBuilder& Builder, TArrayView<const FMLTensorBinding> InInputBindings, TArrayView<const FMLTensorBinding> OutOutputBindings)
-{
 	return 0;
 }
 
-
+/** 
+ * Process tensor bindings and check if we need to create RDG Buffer for CPU binding 
+ * Returns 0 on success, or index of a tensor binding if the tensor binding type is not supported.
+ */
 int FMLInferenceModelRDG::SetTensors(FRDGBuilder& GraphBuilder, FMLTensorBindingArray& OutBindings, FMLIntArray& OutIndices, TArrayView<const FMLTensorBinding> InBindings, TArrayView<const FMLTensorDesc> InTensors)
 {
 	for (int32 Idx = 0; Idx < InBindings.Num(); ++Idx)
@@ -162,7 +186,7 @@ int FMLInferenceModelRDG::SetTensors(FRDGBuilder& GraphBuilder, FMLTensorBinding
 		}
 		else
 		{
-			// Unsupported
+			// Unsupported tensor binding type
 			return Idx;
 		}
 	}
@@ -187,7 +211,7 @@ void FMLInferenceModelRDG::AddTensorUploads_RenderThread(FRDGBuilder& GraphBuild
 		TensorUploadParams->Buffer = RDGBinding.Buffer;
 
 		GraphBuilder.AddPass(
-			RDG_EVENT_NAME("NNXDmlTensorUpload"),
+			RDG_EVENT_NAME("FMLInferenceModelAddTensorUpload"),
 			TensorUploadParams,
 			ERDGPassFlags::Copy | ERDGPassFlags::NeverCull,
 			[InBinding, TensorDesc, TensorUploadParams](FRHICommandListImmediate& RHICmdList)
@@ -217,7 +241,7 @@ void FMLInferenceModelRDG::AddTensorReadbacks_RenderThread(FRDGBuilder& GraphBui
 		TensorReadbackParams->Buffer = RDGBinding.Buffer;
 
 		GraphBuilder.AddPass(
-			RDG_EVENT_NAME("NNXDmlTensorReadback"),
+			RDG_EVENT_NAME("FMLInferenceModelAddTensorReadback"),
 			TensorReadbackParams,
 			ERDGPassFlags::Copy | ERDGPassFlags::NeverCull,
 			[InBinding, TensorDesc, TensorReadbackParams](FRHICommandListImmediate& RHICmdList)
