Cleanup NNXRuntimeORT
- added PLATFORM_WINDOWS for DML EP
- fixed static analysis warnings

#rnx

diff --git a/Engine/Plugins/Experimental/NNX/Source/NNXRuntimeORT/Private/NNXRuntimeORT.cpp b/Engine/Plugins/Experimental/NNX/Source/NNXRuntimeORT/Private/NNXRuntimeORT.cpp
--- a/Engine/Plugins/Experimental/NNX/Source/NNXRuntimeORT/Private/NNXRuntimeORT.cpp
+++ b/Engine/Plugins/Experimental/NNX/Source/NNXRuntimeORT/Private/NNXRuntimeORT.cpp
@@ -1,12 +1,12 @@
 // Copyright Epic Games, Inc. All Rights Reserved.
 #include "NNXRuntimeORT.h"
+
 #include "NeuralTimer.h"
-#include "NNXRuntimeORTProviders.h"
 #include "NNXRuntimeORTUtils.h"
 #include "RedirectCoutAndCerrToUeLog.h"
 
 // NOTE: For now we only have DML on Windows, we should add support for XSX
-//#if PLATFORM_WINDOWS
+#if PLATFORM_WINDOWS
 
 #include "Windows/AllowWindowsPlatformTypes.h"
 #include <unknwn.h>
@@ -16,6 +16,7 @@
 #include "ID3D12DynamicRHI.h"
 #include "D3D12RHIBridge.h"
 #include "DirectML.h"
+#endif
 
 using namespace NNX;
 
@@ -29,10 +30,12 @@ FString FRuntimeORTCuda::GetRuntimeName() const
 	return NNX_RUNTIME_ORT_NAME_CUDA;
 }
 
+#if PLATFORM_WINDOWS
 FString FRuntimeORTDml::GetRuntimeName() const
 {
 	return NNX_RUNTIME_ORT_NAME_DML;
 }
+#endif
 
 EMLRuntimeSupportFlags FRuntimeORTCpu::GetSupportFlags() const
 {
@@ -44,10 +47,12 @@ EMLRuntimeSupportFlags FRuntimeORTCuda::GetSupportFlags() const
 	return EMLRuntimeSupportFlags::GPU;
 }
 
+#if PLATFORM_WINDOWS
 EMLRuntimeSupportFlags FRuntimeORTDml::GetSupportFlags() const
 {
 	return EMLRuntimeSupportFlags::GPU;
 }
+#endif
 
 FMLInferenceModel* FRuntimeORTCpu::CreateInferenceModel(UMLInferenceModel* InModel, const FMLInferenceNNXORTConf& InConf)
 {
@@ -73,6 +78,7 @@ FMLInferenceModel* FRuntimeORTCuda::CreateInferenceModel(UMLInferenceModel* InMo
 	return ORTModel;
 }
 
+#if PLATFORM_WINDOWS
 FMLInferenceModel* FRuntimeORTDml::CreateInferenceModel(UMLInferenceModel* InModel, const FMLInferenceNNXORTConf& InConf)
 {
 	FMLInferenceModelORTDml* ORTModel = new FMLInferenceModelORTDml(&NNXEnvironmentORT, InConf);
@@ -84,6 +90,7 @@ FMLInferenceModel* FRuntimeORTDml::CreateInferenceModel(UMLInferenceModel* InMod
 
 	return ORTModel;
 }
+#endif
 
 FMLInferenceModel* FRuntimeORTCpu::CreateInferenceModel(UMLInferenceModel* InModel)
 {
@@ -97,11 +104,13 @@ FMLInferenceModel* FRuntimeORTCuda::CreateInferenceModel(UMLInferenceModel* InMo
 	return CreateInferenceModel(InModel, ORTInferenceConf);
 }
 
+#if PLATFORM_WINDOWS
 FMLInferenceModel* FRuntimeORTDml::CreateInferenceModel(UMLInferenceModel* InModel)
 {
 	FMLInferenceNNXORTConf ORTInferenceConf;
 	return CreateInferenceModel(InModel, ORTInferenceConf);
 }
+#endif
 
 FMLInferenceModelORT::FMLInferenceModelORT(
 	Ort::Env* InORTEnvironment, 
@@ -357,12 +366,19 @@ bool FMLInferenceModelORTCuda::InitializedAndConfigureMembers()
 
 	SessionOptions->EnableCpuMemArena();
 
-	OrtSessionOptionsAppendExecutionProvider_CUDA(*SessionOptions.Get(), ORTConfiguration.DeviceId);
+	OrtStatusPtr Status = OrtSessionOptionsAppendExecutionProvider_CUDA(*SessionOptions.Get(), ORTConfiguration.DeviceId);
+	if (Status)
+	{
+		UE_LOG(LogNNX, Warning, TEXT("Failed to initialize session options for ORT CUDA EP"));
+		return false;
+	}
 
 	return true;
 }
 
 //-------------
+#if PLATFORM_WINDOWS
+
 FMLInferenceModelORTDml::FMLInferenceModelORTDml(Ort::Env* InORTEnvironment, const FMLInferenceNNXORTConf& InORTConfiguration) :
 	FMLInferenceModelORT(InORTEnvironment, EMLInferenceModelType::GPU, InORTConfiguration)
 {}
@@ -429,7 +445,14 @@ bool FMLInferenceModelORTDml::InitializedAndConfigureMembers()
 	ID3D12CommandQueue* CmdQ = RHI->RHIGetCommandQueue();
 
 	//OrtSessionOptionsAppendExecutionProvider_DML(*SessionOptions.Get(), ORTConfiguration.DeviceId);
-	OrtSessionOptionsAppendExecutionProviderEx_DML(*SessionOptions.Get(), DmlDevice, CmdQ);
+	
+	OrtStatusPtr Status = OrtSessionOptionsAppendExecutionProviderEx_DML(*SessionOptions.Get(), DmlDevice, CmdQ);
+	if (Status)
+	{
+		UE_LOG(LogNNX, Warning, TEXT("Failed to initialize session options for ORT Dml EP"));
+		return false;
+	}
 
 	return true;
 }
+#endif
diff --git a/Engine/Plugins/Experimental/NNX/Source/NNXRuntimeORT/Private/NNXRuntimeORT.h b/Engine/Plugins/Experimental/NNX/Source/NNXRuntimeORT/Private/NNXRuntimeORT.h
new file mode 100644
--- /dev/null
+++ b/Engine/Plugins/Experimental/NNX/Source/NNXRuntimeORT/Private/NNXRuntimeORT.h
@@ -0,0 +1,250 @@
+// Copyright Epic Games, Inc. All Rights Reserved.
+
+#pragma once
+
+#include "CoreMinimal.h"
+#include "NNXCore.h"
+#include "NNXRuntime.h"
+#include "NNXInferenceModel.h"
+#include "NeuralStatistics.h"
+#include "ThirdPartyWarningDisabler.h"
+
+NNI_THIRD_PARTY_INCLUDES_START
+#undef check
+#undef TEXT
+
+#include "core/session/onnxruntime_cxx_api.h"
+
+#include "core/providers/cpu/cpu_provider_factory.h"
+
+#if PLATFORM_WINDOWS
+#include "Windows/MinWindows.h"
+
+#ifdef USE_DML
+#include "core/providers/dml/dml_provider_factory.h"
+#endif
+#endif
+
+NNI_THIRD_PARTY_INCLUDES_END
+
+#define NNX_RUNTIME_ORT_NAME_CPU TEXT("NNXRuntimeORTCpu")
+#define NNX_RUNTIME_ORT_NAME_DML TEXT("NNXRuntimeORTDml")
+#define NNX_RUNTIME_ORT_NAME_CUDA TEXT("NNXRuntimeORTCuda")
+
+namespace NNX
+{
+	struct FMLInferenceNNXORTConf
+	{
+		FMLInferenceNNXORTConf() {};
+		FMLInferenceNNXORTConf(
+			const uint32 InDeviceId, 
+			uint32 InNumberOfThreads = 2, 
+			GraphOptimizationLevel InOptimizationLevel = GraphOptimizationLevel::ORT_ENABLE_ALL
+		) : DeviceId(InDeviceId), NumberOfThreads(InNumberOfThreads), OptimizationLevel(InOptimizationLevel)
+		{
+		};
+
+		uint32 DeviceId = 0;
+		uint32 NumberOfThreads = 2;
+		GraphOptimizationLevel OptimizationLevel = GraphOptimizationLevel::ORT_ENABLE_ALL;
+	};
+
+	class FRuntimeORTCpu : public IRuntime
+	{
+	public:
+		Ort::Env NNXEnvironmentORT;
+		FRuntimeORTCpu() = default;
+		virtual ~FRuntimeORTCpu() = default;
+		virtual FString GetRuntimeName() const override;
+		virtual EMLRuntimeSupportFlags GetSupportFlags() const override;
+		virtual FMLInferenceModel* CreateInferenceModel(UMLInferenceModel* InModel);
+		FMLInferenceModel* CreateInferenceModel(UMLInferenceModel* InModel, const FMLInferenceNNXORTConf& InConf);
+	};
+
+	class FRuntimeORTCuda : public IRuntime
+	{
+	public:
+		Ort::Env NNXEnvironmentORT;
+		FRuntimeORTCuda() = default;
+		virtual ~FRuntimeORTCuda() = default;
+		virtual FString GetRuntimeName() const override;
+		virtual EMLRuntimeSupportFlags GetSupportFlags() const override;
+		virtual FMLInferenceModel* CreateInferenceModel(UMLInferenceModel* InModel);
+		FMLInferenceModel* CreateInferenceModel(UMLInferenceModel* InModel, const FMLInferenceNNXORTConf& InConf);
+	};
+
+#if PLATFORM_WINDOWS
+	class FRuntimeORTDml : public IRuntime
+	{
+	public:
+		Ort::Env NNXEnvironmentORT;
+		FRuntimeORTDml() = default;
+		virtual ~FRuntimeORTDml() = default;
+		virtual FString GetRuntimeName() const override;
+		virtual EMLRuntimeSupportFlags GetSupportFlags() const override;
+		virtual FMLInferenceModel* CreateInferenceModel(UMLInferenceModel* InModel);
+		FMLInferenceModel* CreateInferenceModel(UMLInferenceModel* InModel, const FMLInferenceNNXORTConf& InConf);
+	};
+#endif
+
+	static TUniquePtr<FRuntimeORTCpu> GORTCPURuntime;
+	static TUniquePtr<FRuntimeORTCuda> GORTCUDARuntime;
+
+#if PLATFORM_WINDOWS
+	static TUniquePtr<FRuntimeORTDml> GORTDMLRuntime;
+#endif
+
+	inline TUniquePtr<FRuntimeORTCpu> FRuntimeORTCPUCreate()
+	{
+		Ort::InitApi();
+		TUniquePtr<FRuntimeORTCpu> Runtime = MakeUnique<FRuntimeORTCpu>();
+
+		return Runtime;
+	};
+	inline TUniquePtr<FRuntimeORTCuda> FRuntimeORTCUDACreate()
+	{
+		Ort::InitApi();
+		TUniquePtr<FRuntimeORTCuda> Runtime = MakeUnique<FRuntimeORTCuda>();
+
+		return Runtime;
+	};
+
+#if PLATFORM_WINDOWS
+	inline TUniquePtr<FRuntimeORTDml> FRuntimeORTDMLCreate()
+	{
+		Ort::InitApi();
+		TUniquePtr<FRuntimeORTDml> Runtime = MakeUnique<FRuntimeORTDml>();
+
+		return Runtime;
+	};
+#endif
+
+	inline IRuntime* FRuntimeORTCPUStartup()
+	{
+		if (!GORTCPURuntime)
+		{
+			GORTCPURuntime = FRuntimeORTCPUCreate();
+		}
+		return GORTCPURuntime.Get();
+	};
+
+	inline IRuntime* FRuntimeORTCUDAStartup()
+	{
+		if (!GORTCUDARuntime)
+		{
+			GORTCUDARuntime = FRuntimeORTCUDACreate();
+		}
+		return GORTCUDARuntime.Get();
+	};
+
+#if PLATFORM_WINDOWS
+	inline IRuntime* FRuntimeORTDMLStartup()
+	{
+		if (!GORTDMLRuntime)
+		{
+			GORTDMLRuntime = FRuntimeORTDMLCreate();
+		}
+		return GORTDMLRuntime.Get();
+	};
+#endif
+
+	inline void FRuntimeORTCPUShutdown()
+	{
+		if (GORTCPURuntime)
+		{
+			GORTCPURuntime.Release();
+		}
+	};
+
+	inline void FRuntimeORTCUDAShutdown()
+	{
+		if (GORTCUDARuntime)
+		{
+			GORTCUDARuntime.Release();
+		}
+	};
+
+#if PLATFORM_WINDOWS
+	inline void FRuntimeORTDMLShutdown()
+	{
+		if (GORTDMLRuntime)
+		{
+			GORTDMLRuntime.Release();
+		}
+	};
+#endif
+
+	class FMLInferenceModelORT : public FMLInferenceModel
+	{
+
+	public:
+		virtual ~FMLInferenceModelORT() = default;
+
+		bool Init(const UMLInferenceModel* InferenceModel);
+		bool IsLoaded() const;
+
+		virtual int Run(TArrayView<const FMLTensorBinding> InInputBindingTensors, TArrayView<const FMLTensorBinding> OutOutputBindingTensors);
+
+		float GetLastRunTimeMSec() const;
+		FNeuralStatistics GetRunStatistics() const;
+		FNeuralStatistics GetInputMemoryTransferStats() const;
+		void ResetStats();
+
+	protected:
+		FMLInferenceModelORT(Ort::Env* InORTEnvironment, EMLInferenceModelType InType, const FMLInferenceNNXORTConf& InORTConfiguration);
+
+		bool bIsLoaded = false;
+		bool bHasRun = false;
+
+		/** ORT-related variables */
+		Ort::Env* ORTEnvironment;
+		FMLInferenceNNXORTConf ORTConfiguration;
+		TUniquePtr<Ort::Session> Session;
+		TUniquePtr<Ort::AllocatorWithDefaultOptions> Allocator;
+		TUniquePtr<Ort::SessionOptions> SessionOptions;
+		TUniquePtr<Ort::MemoryInfo> AllocatorInfo;
+
+		/** IO ORT-related variables */
+		TArray<ONNXTensorElementDataType> InputTensorsORTType;
+		TArray<ONNXTensorElementDataType> OutputTensorsORTType;
+
+		TArray<const char*> InputTensorNames;
+		TArray<const char*> OutputTensorNames;
+
+		virtual bool InitializedAndConfigureMembers();
+		bool ConfigureTensors(const bool InIsInput);
+
+		/**
+		 * Statistics-related members used for GetLastRunTimeMSec(), GetRunStatistics(), GetInputMemoryTransferStats(), ResetStats().
+		 */
+		FNeuralStatisticsEstimator RunStatisticsEstimator;
+		FNeuralStatisticsEstimator InputTransferStatisticsEstimator;
+	};
+
+	class FMLInferenceModelORTCpu : public NNX::FMLInferenceModelORT
+	{
+	public:
+		FMLInferenceModelORTCpu(Ort::Env* InORTEnvironment, const FMLInferenceNNXORTConf& InORTConfiguration);
+
+		virtual bool InitializedAndConfigureMembers() override;
+	};
+
+#if PLATFORM_WINDOWS
+	class FMLInferenceModelORTDml : public NNX::FMLInferenceModelORT
+	{
+	public:
+		FMLInferenceModelORTDml(Ort::Env* InORTEnvironment, const FMLInferenceNNXORTConf& InORTConfiguration);
+
+		virtual bool InitializedAndConfigureMembers() override;
+	};
+#endif
+
+	class FMLInferenceModelORTCuda : public NNX::FMLInferenceModelORT
+	{
+	public:
+		FMLInferenceModelORTCuda(Ort::Env* InORTEnvironment, const FMLInferenceNNXORTConf& InORTConfiguration);
+
+		virtual bool InitializedAndConfigureMembers() override;
+	};
+
+}
\ No newline at end of file
diff --git a/Engine/Plugins/Experimental/NNX/Source/NNXRuntimeORT/Private/NNXRuntimeORTProviders.h b/Engine/Plugins/Experimental/NNX/Source/NNXRuntimeORT/Private/NNXRuntimeORTProviders.h
new file mode 100644
--- /dev/null
+++ b/Engine/Plugins/Experimental/NNX/Source/NNXRuntimeORT/Private/NNXRuntimeORTProviders.h
@@ -0,0 +1,48 @@
+// Copyright Epic Games, Inc. All Rights Reserved.
+// Copyright (c) Microsoft Corporation. All rights reserved.
+// Licensed under the MIT License.
+
+#pragma once
+
+#include "ThirdPartyWarningDisabler.h"
+NNI_THIRD_PARTY_INCLUDES_START
+#undef check
+#undef TEXT
+
+#include "core/providers/cpu/cpu_provider_factory.h"
+
+#ifdef USE_DNNL
+#include "core/providers/dnnl/dnnl_provider_factory.h"
+#endif
+#ifdef USE_NUPHAR
+#include "core/providers/nuphar/nuphar_provider_factory.h"
+#endif
+#ifdef USE_TVM
+#include "core/providers/tvm/tvm_provider_factory.h"
+#endif
+#ifdef USE_TENSORRT
+#include "core/providers/tensorrt/tensorrt_provider_factory.h"
+#endif
+#ifdef USE_OPENVINO
+#include "core/providers/openvino/openvino_provider_factory.h"
+#endif
+#ifdef USE_NNAPI
+#include "core/providers/nnapi/nnapi_provider_factory.h"
+#endif
+#ifdef USE_COREML
+#include "core/providers/coreml/coreml_provider_factory.h"
+#endif
+#ifdef USE_DML
+#include "core/providers/dml/dml_provider_factory.h"
+#endif
+#ifdef USE_ACL
+#include "core/providers/acl/acl_provider_factory.h"
+#endif
+#ifdef USE_ARMNN
+#include "core/providers/armnn/armnn_provider_factory.h"
+#endif
+#ifdef USE_MIGRAPHX
+#include "core/providers/migraphx/migraphx_provider_factory.h"
+#endif
+
+NNI_THIRD_PARTY_INCLUDES_END
\ No newline at end of file
diff --git a/Engine/Plugins/Experimental/NNX/Source/NNXRuntimeORT/Public/NNXRuntimeORT.h b/Engine/Plugins/Experimental/NNX/Source/NNXRuntimeORT/Public/NNXRuntimeORT.h
deleted file mode 100644
diff --git a/Engine/Plugins/Experimental/NNX/Source/NNXRuntimeORT/Public/NNXRuntimeORTProviders.h b/Engine/Plugins/Experimental/NNX/Source/NNXRuntimeORT/Public/NNXRuntimeORTProviders.h
deleted file mode 100644
