Ported Conv shader.

diff --git a/Engine/Plugins/Experimental/NNX/Source/NNXHlslShaders/Private/NNXConvCS.cpp b/Engine/Plugins/Experimental/NNX/Source/NNXHlslShaders/Private/NNXConvCS.cpp
new file mode 100644
--- /dev/null
+++ b/Engine/Plugins/Experimental/NNX/Source/NNXHlslShaders/Private/NNXConvCS.cpp
@@ -0,0 +1,288 @@
+// Copyright Epic Games, Inc. All Rights Reserved.
+
+#include "NNXConvCS.h"
+#include "NNXCore.h"
+
+namespace
+{
+
+TArray<int32> GetXBlockShape(TArray<int32> GroupShape, TArray<int32> WShape, TArray<int32> Dilations, TArray<int32> Strides)
+{
+	check(WShape.Num() > 2);
+	check(GroupShape.Num() == WShape.Num() - 2);
+	check(Dilations.Num() == 0 || Dilations.Num() == GroupShape.Num());
+	check(Strides.Num() == 0 || Strides.Num() == GroupShape.Num());
+	
+	TArray<int32> Result;
+	Result.SetNumUninitialized(GroupShape.Num());
+	for (int32 i = 0; i < GroupShape.Num(); i++)
+	{
+		int32 DilatedKernelSize = (i < Dilations.Num() ? Dilations[i] : 1) * (WShape[2 + i] - 1) + 1;
+		Result[i] = DilatedKernelSize + (GroupShape[i] - 1) * (i < Strides.Num() ? Strides[i] : 1); 
+	}
+	return Result;
+}
+
+TArray<int32> GetPadding(TArray<int32> XShape, TArray<int32> WShape, EConvAutoPad AutoPad, TArray<int32> Dilations, TArray<int32> Strides, TArray<int32> Pads)
+{
+	check(XShape.Num() > 2);
+	check(WShape.Num() == XShape.Num());
+	check(Dilations.Num() == 0 || Dilations.Num() == WShape.Num() - 2);
+	check(Strides.Num() == 0 || Strides.Num() == WShape.Num() - 2);
+	check(AutoPad != EConvAutoPad::NOTSET || Pads.Num() == 2 * (WShape.Num() - 2));
+
+	int32 NumDimensions = XShape.Num() - 2;
+	TArray<int32> Result;
+	Result.Init(0, 2 * NumDimensions);
+
+	if (AutoPad == EConvAutoPad::NOTSET)
+	{
+		return Pads;
+	}
+	else if (AutoPad == EConvAutoPad::VALID)
+	{
+		return Result;
+	}
+
+	for (int32 DimensionIndex = 0; DimensionIndex < NumDimensions; DimensionIndex++)
+	{
+		int32 DilatedKernelSize = (DimensionIndex < Dilations.Num() ? Dilations[DimensionIndex] : 1) * (WShape[DimensionIndex + 2] - 1) + 1;
+		int32 TotalPad = ((int32)(((XShape[DimensionIndex + 2] + (DimensionIndex < Strides.Num() ? Strides[DimensionIndex] : 1) - 1) / (DimensionIndex < Strides.Num() ? Strides[DimensionIndex] : 1)) - 1)) * (DimensionIndex < Strides.Num() ? Strides[DimensionIndex] : 1) + DilatedKernelSize - XShape[DimensionIndex + 2];
+		if (AutoPad == EConvAutoPad::SAME_LOWER)
+		{
+			Result[DimensionIndex] = (TotalPad + 1) / 2;
+		}	
+		else
+		{
+			Result[DimensionIndex] = TotalPad / 2;
+		}
+		Result[NumDimensions + DimensionIndex] = TotalPad - Result[DimensionIndex];
+	}
+
+	return Result;
+}
+
+int32 GetNumThreadsPerGroup(EConvGroupSize GroupSize)
+{
+	int32 NumThreadsPerGroup = 128;
+	switch (GroupSize)
+	{
+		case EConvGroupSize::Size128:
+			NumThreadsPerGroup = 128;
+			break;
+		case EConvGroupSize::Size256:
+			NumThreadsPerGroup = 256;
+			break;
+		case EConvGroupSize::Size512:
+			NumThreadsPerGroup = 512;
+			break;
+		default:
+			NumThreadsPerGroup = 128;
+			break;
+	}
+	check(FMath::Log2((float)NumThreadsPerGroup) == FMath::Floor(FMath::Log2((float)NumThreadsPerGroup)));
+	return NumThreadsPerGroup;
+}
+
+TArray<int32> GetGridShape(TArray<int32> YShape, TArray<int32> GroupShape)
+{
+	check(YShape.Num() > 2);
+	check(YShape.Num() == (GroupShape.Num() + 2));
+
+	TArray<int32> Result;
+	Result.SetNumUninitialized(YShape.Num() - 2);
+	for (int32 i = 2; i < YShape.Num(); i++)
+	{
+		Result[i-2] = FMath::DivideAndRoundUp(YShape[i], GroupShape[i-2]);
+	}
+
+	return Result;
+}
+
+} // namespace
+
+void FMLConvCS::ModifyCompilationEnvironment(const FGlobalShaderPermutationParameters& InParameters, FShaderCompilerEnvironment& OutEnvironment)
+{
+	FGlobalShader::ModifyCompilationEnvironment(InParameters, OutEnvironment);
+	OutEnvironment.SetDefine(TEXT("MAX_NUM_DIMENSIONS"), NNXRT_CONV_MAX_NUM_DIMENSIONS);
+}
+
+TArray<int32> FMLConvCS::GetOutputShape(TArray<int32> XShape, TArray<int32> WShape, EConvAutoPad AutoPad, TArray<int32> Dilations, TArray<int32> Strides, TArray<int32> Pads)
+{
+	check(XShape.Num() > 2);
+	check(WShape.Num() == XShape.Num());
+	check(Dilations.Num() == 0 || Dilations.Num() == WShape.Num() - 2);
+	check(Strides.Num() == 0 || Strides.Num() == WShape.Num() - 2);
+	check(Pads.Num() == 0 || Pads.Num() == 2 * (WShape.Num() - 2));
+
+	TArray<int32> Padding = GetPadding(XShape, WShape, AutoPad, Dilations, Strides, Pads);
+
+	TArray<int32> Result;
+	Result.SetNumUninitialized(XShape.Num());
+	Result[0] = XShape[0];
+	Result[1] = WShape[0];
+
+	int32 NumDimensions = XShape.Num() - 2;
+	for (int32 DimensionIndex = 0; DimensionIndex < NumDimensions; DimensionIndex++)
+	{
+		int32 PaddedSize = XShape[2 + DimensionIndex] + Padding[DimensionIndex] + Padding[NumDimensions + DimensionIndex];
+		int32 DilatedKernelSize = (DimensionIndex < Dilations.Num() ? Dilations[DimensionIndex] : 1) * (WShape[2 + DimensionIndex] - 1) + 1;
+		float PaddeSizeMinusDilatedKernelSize = (float)(PaddedSize - DilatedKernelSize);
+		int32 OutputSize = (int32)(PaddeSizeMinusDilatedKernelSize / (DimensionIndex < Strides.Num() ? (float)Strides[DimensionIndex] : 1.0) + 1);
+		Result[2 + DimensionIndex] = OutputSize;
+	}
+
+	return Result;
+}
+
+void FMLConvCS::FillInParameters(EConvGroupSize GroupSize, TArray<int32> XShape, TArray<int32> WShape, bool HasB, EConvAutoPad AutoPad, int Group, TArray<int32> Dilations, TArray<int32> Strides, TArray<int32> Pads, FMLConvCS::FParameters& Parameters)
+{
+	check(XShape.Num() > 2);
+	check(WShape.Num() == XShape.Num());
+	check(Dilations.Num() == 0 || Dilations.Num() == WShape.Num() - 2);
+	check(Strides.Num() == 0 || Strides.Num() == WShape.Num() - 2);
+	check(Pads.Num() == 0 || Pads.Num() == 2 * (WShape.Num() - 2));
+	check(GetNumReadsPerThread(GroupSize, WShape, Dilations, Strides) >= 0)
+	check(WShape[0] > 0)
+	check(WShape[1] > 0)
+
+	int32 NumDimensions = XShape.Num() - 2;
+
+	TArray<int32> Padding = GetPadding(XShape, WShape, AutoPad, Dilations, Strides, Pads);
+	TArray<int32> GroupShape = GetGroupShape(GroupSize, NumDimensions);
+	TArray<int32> YShape = GetOutputShape(XShape, WShape, AutoPad, Dilations, Strides, Pads);
+	TArray<int32> GridShape = GetGridShape(YShape, GroupShape);
+	TArray<int32> XBlockShape = GetXBlockShape(GroupShape, WShape, Dilations, Strides);
+	
+	int32 GroupStride = 1;
+	int32 GroupThreadStride = 1;
+	int32 XBlockSize = 1;
+	int32 YMemoryStride = 1;
+	int32 XMemoryStride = 1;
+	int32 WChannelSize = 1;
+	for (int32 i = NumDimensions-1; i >= 0; i--)
+	{
+		int32 Stride = i < Strides.Num() ? Strides[i] : 1;
+		int32 Dilation = i < Dilations.Num() ? Dilations[i] : 1;
+		Parameters.Dilation_Stride_XBlockStartOffset_DilationXBlockStride[i] = FIntVector4(Dilation, Stride, -Padding[i], Dilation * XBlockSize);
+		Parameters.GroupStride_GroupShape_GroupThreadStride_StrideXBlockStride[i] = FIntVector4(GroupStride, GroupShape[i], GroupThreadStride, Stride * XBlockSize);
+		Parameters.YDimension_YMemoryStride_XDimension_XMemoryStride[i] = FIntVector4(YShape[2 + i], YMemoryStride, XShape[2 + i], XMemoryStride);
+		Parameters.XBlockStartStride_XBlockStride_WDimension_WDimensionDilationXBlockStride[i] = FIntVector4(GroupShape[i] * Stride, XBlockSize, WShape[2 + i], WShape[2 + i] * Dilation * XBlockSize);
+		Parameters.OneDiv_GroupStride_GroupThreadStride_XBlockStride[i] = FVector4f(1.0/((float)GroupStride), 1.0/((float)GroupThreadStride), 1.0/((float)XBlockSize), 0.0);
+
+		GroupStride *= GridShape[i];
+		GroupThreadStride *= GroupShape[i];
+		XBlockSize *= XBlockShape[i];
+		YMemoryStride *= YShape[2 + i];
+		XMemoryStride *= XShape[2 + i];
+		WChannelSize *= WShape[2 + i];
+	}
+	
+	Parameters.NumWChannels = WShape[1];
+
+	Parameters.YBatchStride = YShape[1] * YMemoryStride;
+	Parameters.YOutputKernelStride = YMemoryStride;
+	
+	Parameters.XBatchStride = XShape[1] * XMemoryStride;
+	Parameters.XChannelStride = XMemoryStride;
+
+	Parameters.XBlockSize = XBlockSize;
+
+	Parameters.NumChannelsPerBatch = FMath::Min((int32)((float)GroupThreadStride / (float)WChannelSize), WShape[1]);
+	check(Parameters.NumChannelsPerBatch > 0)
+	Parameters.NumChannelBatches = FMath::DivideAndRoundUp(WShape[1], Parameters.NumChannelsPerBatch);
+	
+	Parameters.WOutputKernelStride = WShape[1] * WChannelSize;
+	Parameters.WChannelBatchSize = Parameters.NumChannelsPerBatch * WChannelSize;
+	Parameters.WChannelSize = WChannelSize;
+
+	Parameters.GroupsDivM = (float)Group / (float)WShape[0];
+}
+
+int32 FMLConvCS::GetNumReadsPerThread(EConvGroupSize GroupSize, TArray<int32> WShape, TArray<int32> Dilations, TArray<int32> Strides)
+{
+	check(WShape.Num() > 2);
+	check(Dilations.Num() == 0 || Dilations.Num() == WShape.Num() - 2);
+	check(Strides.Num() == 0 || Strides.Num() == WShape.Num() - 2);
+
+	int32 NumDimensions = WShape.Num() - 2;
+
+	TArray<int32> GroupShape = GetGroupShape(GroupSize, NumDimensions);
+	int32 NumThreadsPerGroup = 1;
+	for (int32 i = 0; i < NumDimensions; i++)
+	{
+		NumThreadsPerGroup *= GroupShape[i];
+	}
+
+	TArray<int32> XBlockShape = GetXBlockShape(GroupShape, WShape, Dilations, Strides);
+	int32 NumXBlockElements = 1;
+	for (int32 i = 0; i < NumDimensions; i++)
+	{
+		NumXBlockElements *= XBlockShape[i];
+	}
+
+	int32 NumReads = FMath::DivideAndRoundUp(NumXBlockElements, NumThreadsPerGroup);
+	int32 NumReadsPow2 = FMath::Max(FMath::RoundToPositiveInfinity(FMath::Log2((float)NumReads)), FConvConstants::MIN_NUM_READS_PER_THREAD_POW2);
+
+	if(NumReadsPow2 <= FConvConstants::MAX_NUM_READS_PER_THREAD_POW2)
+	{
+		return NumReadsPow2;
+	}
+
+	return -1;
+}
+
+TArray<int32> FMLConvCS::GetGroupShape(EConvGroupSize GroupSize, int32 NumDimensions)
+{
+	check(NumDimensions > 0);
+
+	int32 NumThreadsPerGroup = GetNumThreadsPerGroup(GroupSize);
+
+	int32 Power = (int32)FMath::Log2((float)NumThreadsPerGroup);
+	int32 MinPowerPerDim = (int32)((float)Power / (float)NumDimensions);
+	int32 PowerReminder = Power - NumDimensions * MinPowerPerDim;
+
+	TArray<int32> Result;
+	Result.Init((int32)FMath::Pow((float)2.0, (float)MinPowerPerDim), NumDimensions);
+	for (int32 i = 0; i < PowerReminder; i++)
+	{
+		Result[NumDimensions-1-i] *= 2;
+	}
+	return Result;
+}
+
+FIntVector FMLConvCS::GetGroupCount(TArray<int32> YShape, TArray<int32> GroupShape)
+{
+	check(YShape.Num() > 2);
+	check(YShape.Num() == (GroupShape.Num() + 2));
+
+	int32 ThreadGroupCountValueX = 1;
+	for (int32 i = 2; i < YShape.Num(); i++)
+	{
+		ThreadGroupCountValueX *= FMath::DivideAndRoundUp(YShape[i], GroupShape[i-2]);
+	}
+
+	return FIntVector(ThreadGroupCountValueX, YShape[1], YShape[0]);
+}
+
+EConvGroupSize FMLConvCS::GetMinimalGroupSize(TArray<int32> WShape)
+{
+	int32 NumDimensions = WShape.Num() - 2;
+	int32 WChannelSize = 1;
+	for (int32 i = 0; i < NumDimensions; i++)
+	{
+		WChannelSize *= WShape[2 + i];
+	}
+	
+	for (int32 i = 0; i < (int32)EConvGroupSize::MAX; i++)
+	{
+		if (GetNumThreadsPerGroup((EConvGroupSize)i) >= WChannelSize) 
+		{
+			return (EConvGroupSize)i;
+		}
+	}
+	
+	return EConvGroupSize::MAX;
+}
+
+IMPLEMENT_GLOBAL_SHADER(FMLConvCS, "/NNX/ConvOp.usf", "main", SF_Compute);
\ No newline at end of file
diff --git a/Engine/Plugins/Experimental/NNX/Source/NNXHlslShaders/Public/NNXConvCS.h b/Engine/Plugins/Experimental/NNX/Source/NNXHlslShaders/Public/NNXConvCS.h
new file mode 100644
--- /dev/null
+++ b/Engine/Plugins/Experimental/NNX/Source/NNXHlslShaders/Public/NNXConvCS.h
@@ -0,0 +1,81 @@
+// Copyright Epic Games, Inc. All Rights Reserved.
+
+#pragma once
+
+#include "GlobalShader.h"
+#include "NNXShaderParameters.h"
+
+enum class EConvAlgorithm : uint8
+{
+	SharedMemory = 0,
+	MAX
+};
+
+enum class EConvGroupSize : uint8
+{
+	Size128 = 0,
+	Size256,
+	Size512,
+	MAX
+};
+
+enum class EConvAutoPad : uint8
+{
+	NOTSET = 0,// Use pad values passed in the array
+	SAME_UPPER,// Auto-pad to match input and output shape with potetnial extra padding at the end
+	SAME_LOWER,// Auto-pad to match input and output shape with potetnial extra padding at the beginning
+	VALID,// Set all paddings to zero
+	MAX
+};
+
+class FConvConstants
+{
+public:
+	static const int32 MIN_NUM_READS_PER_THREAD_POW2{1};
+	static const int32 MAX_NUM_READS_PER_THREAD_POW2{3};
+};
+
+class NNXHLSLSHADERS_API FMLConvCS : public FGlobalShader
+{
+	DECLARE_GLOBAL_SHADER(FMLConvCS);
+	SHADER_USE_PARAMETER_STRUCT(FMLConvCS, FGlobalShader)
+
+	class FConvAlgorithm : SHADER_PERMUTATION_ENUM_CLASS("ALGORITHM", EConvAlgorithm);
+	class FConvGroupSize : SHADER_PERMUTATION_ENUM_CLASS("GROUP_SIZE", EConvGroupSize);
+	class FConvNumDimensions : SHADER_PERMUTATION_RANGE_INT("NUM_DIMENSIONS", 1, NNXRT_CONV_MAX_NUM_DIMENSIONS);
+	class FConvNumReadsPerThread : SHADER_PERMUTATION_RANGE_INT("NUM_READS_PER_THREAD_POW2", FConvConstants::MIN_NUM_READS_PER_THREAD_POW2, FConvConstants::MAX_NUM_READS_PER_THREAD_POW2);
+	class FConvHasB : SHADER_PERMUTATION_BOOL("HAS_B");
+	using FPermutationDomain = TShaderPermutationDomain<FConvAlgorithm, FConvGroupSize, FConvNumDimensions, FConvNumReadsPerThread, FConvHasB>;
+
+public:
+	NNXRT_CONV_PARAMETER_STRUCT();
+
+	static void ModifyCompilationEnvironment(const FGlobalShaderPermutationParameters& InParameters, FShaderCompilerEnvironment& OutEnvironment);
+
+public:
+	static TArray<int32> GetOutputShape(TArray<int32> XShape, TArray<int32> WShape, EConvAutoPad AutoPad, TArray<int32> Dilations, TArray<int32> Strides, TArray<int32> Pads);
+
+	static void FillInParameters(EConvGroupSize GroupSize, TArray<int32> XShape, TArray<int32> WShape, bool HasB, EConvAutoPad AutoPad, int Group, TArray<int32> Dilations, TArray<int32> Strides, TArray<int32> Pads, FMLConvCS::FParameters& Parameters);
+
+	static int32 GetNumReadsPerThread(EConvGroupSize GroupSize, TArray<int32> WShape, TArray<int32> Dilations, TArray<int32> Strides);
+
+	/**
+	 * @brief Computes the group shape such that all dimension have roughly equal sizes.
+	 *
+	 * @param GroupSize The enum indicating the number of threads contained by a single group.
+	 * @param NumDimensions The number of dimensions.
+	 * @return TArray<int32> An array of size \p NumDimensions containing the number of threads in each dimension to form a volume of a total number of threads indicated by \p GroupSize
+	 */
+	static TArray<int32> GetGroupShape(EConvGroupSize GroupSize, int32 NumDimensions);
+
+	/**
+	 * @brief Get the group count vector used to launch the gpu shader thread groups
+	 *
+	 * @param YShape The shape of the output as computed by GetOutputShape()
+	 * @param YShape The shape of a single thread group as computed by GetGroupShape()
+	 * @return FIntVector The number of thread groups to instantiate. z corresponds to the batch and y to the output kernel.
+	 */
+	static FIntVector GetGroupCount(TArray<int32> YShape, TArray<int32> GroupShape);
+
+	static EConvGroupSize GetMinimalGroupSize(TArray<int32> WShape);
+};
\ No newline at end of file
diff --git a/Engine/Plugins/Experimental/NNX/Source/NNXHlslShaders/Public/NNXShaderParameters.h b/Engine/Plugins/Experimental/NNX/Source/NNXHlslShaders/Public/NNXShaderParameters.h
--- a/Engine/Plugins/Experimental/NNX/Source/NNXHlslShaders/Public/NNXShaderParameters.h
+++ b/Engine/Plugins/Experimental/NNX/Source/NNXHlslShaders/Public/NNXShaderParameters.h
@@ -115,3 +115,35 @@ NNXRT_ELEMENTWISEBINARY_PARAMETER_STRUCT()
 #ifdef NNXRT_RDG_DML
 NNXRT_GEMM_PARAMETER_STRUCT()
 #endif
+
+//Conv
+#define NNXRT_CONV_MAX_NUM_DIMENSIONS 4
+
+#define NNXRT_CONV_PARAMETER_STRUCT() \
+	NNXRT_BEGIN_SHADER_PARAMETER_STRUCT(FMLConvParameters) \
+		NNXRT_RDG_BUFFER_SRV(Buffer<float>, X) \
+		NNXRT_RDG_BUFFER_SRV(Buffer<float>, W) \
+		NNXRT_RDG_BUFFER_UAV(RWBuffer<float>, Y) \
+		NNXRT_RDG_BUFFER_SRV(Buffer<float>, B) \
+		SHADER_PARAMETER_ARRAY(FIntVector4, Dilation_Stride_XBlockStartOffset_DilationXBlockStride, [NNXRT_CONV_MAX_NUM_DIMENSIONS]) \
+		SHADER_PARAMETER_ARRAY(FIntVector4, GroupStride_GroupShape_GroupThreadStride_StrideXBlockStride, [NNXRT_CONV_MAX_NUM_DIMENSIONS]) \
+		SHADER_PARAMETER_ARRAY(FIntVector4, YDimension_YMemoryStride_XDimension_XMemoryStride, [NNXRT_CONV_MAX_NUM_DIMENSIONS]) \
+		SHADER_PARAMETER_ARRAY(FIntVector4, XBlockStartStride_XBlockStride_WDimension_WDimensionDilationXBlockStride, [NNXRT_CONV_MAX_NUM_DIMENSIONS]) \
+		SHADER_PARAMETER_ARRAY(FVector4f, OneDiv_GroupStride_GroupThreadStride_XBlockStride, [NNXRT_CONV_MAX_NUM_DIMENSIONS]) \
+		SHADER_PARAMETER(int32, NumWChannels) \
+		SHADER_PARAMETER(int32, YBatchStride) \
+		SHADER_PARAMETER(int32, YOutputKernelStride) \
+		SHADER_PARAMETER(int32, XBatchStride) \
+		SHADER_PARAMETER(int32, XChannelStride) \
+		SHADER_PARAMETER(int32, XBlockSize) \
+		SHADER_PARAMETER(int32, NumChannelBatches) \
+		SHADER_PARAMETER(int32, NumChannelsPerBatch) \
+		SHADER_PARAMETER(int32, WOutputKernelStride) \
+		SHADER_PARAMETER(int32, WChannelBatchSize) \
+		SHADER_PARAMETER(int32, WChannelSize) \
+		SHADER_PARAMETER(float, GroupsDivM) \
+	NNXRT_END_SHADER_PARAMETER_STRUCT()
+
+#ifdef NNXRT_RDG_DML
+NNXRT_CONV_PARAMETER_STRUCT()
+#endif
\ No newline at end of file
diff --git a/Engine/Plugins/Experimental/NNX/Source/NNXHlslShaders/Shaders/ConvOp.usf b/Engine/Plugins/Experimental/NNX/Source/NNXHlslShaders/Shaders/ConvOp.usf
new file mode 100644
--- /dev/null
+++ b/Engine/Plugins/Experimental/NNX/Source/NNXHlslShaders/Shaders/ConvOp.usf
@@ -0,0 +1,294 @@
+// Copyright Epic Games, Inc. All Rights Reserved.
+
+#include "/Engine/Public/Platform.ush"
+
+#define WORK_TYPE float
+#define BUFFER_TYPE float
+#define BUFFER_TO_WORK_TYPE(x) x
+#define WORK_TO_BUFFER_TYPE(x) x
+
+// Must correspond to EConvGroupSize defined in ConvCS.h
+#if GROUP_SIZE == 0
+#define NUM_GROUP_THREADS 128
+#elif GROUP_SIZE == 1
+#define NUM_GROUP_THREADS 256
+#elif GROUP_SIZE == 2
+#define NUM_GROUP_THREADS 512
+#endif
+
+Buffer<BUFFER_TYPE> X; // N x C x XD0 x ... x XDi
+Buffer<BUFFER_TYPE> W; // M x C/GROUPS x WD0 x ... x WDi
+RWBuffer<BUFFER_TYPE> Y;  // N x M x YD0 x ... x YDi
+Buffer<BUFFER_TYPE> B; // B0, ..., BM
+
+// x: Dilation0, ...., Dilationi
+// y: Strides0, ...., Stridesi
+// z: Pad0Begin, ...., PadiBegin
+int4 Dilation_Stride_XBlockStartOffset_DilationXBlockStride[MAX_NUM_DIMENSIONS];
+
+// x: Element j contains PROD(GDj+1,..,GDi) with the last element containing 1, with GDi being the number of thread groups contained by a dimension i
+// y: Number of threads in each dimension of a group
+// z: Element j contains PROD(GTDj+1,..,GTDi) with the last element containing 1, with GTDi being the number of threads contained by a dimension i inside a group
+int4 GroupStride_GroupShape_GroupThreadStride_StrideXBlockStride[MAX_NUM_DIMENSIONS];
+
+// x: YD0, ..., YDi
+// y: Element j contains PROD(YDj+1,..,YDi) with the last element containing 1
+// z: XD0, ..., XDi
+// w: Element j contains PROD(XDj+1,..,XDi) with the last element containing 1
+int4 YDimension_YMemoryStride_XDimension_XMemoryStride[MAX_NUM_DIMENSIONS];
+
+// x: Number of elements in each dimension of a X block
+// y: The strides of the X block to be loaded
+// z: WD0, ..., WDi
+int4 XBlockStartStride_XBlockStride_WDimension_WDimensionDilationXBlockStride[MAX_NUM_DIMENSIONS];
+
+// x: 1/GroupStride
+// y: 1/GroupThreadStride
+// z: 1/XBlockStride
+float4 OneDiv_GroupStride_GroupThreadStride_XBlockStride[MAX_NUM_DIMENSIONS];
+
+int NumWChannels; // WShape[1]
+
+int YBatchStride;
+int YOutputKernelStride;
+
+int XBatchStride; // The number of elements in each bacth of X
+int XChannelStride; // The number of elements in each channel of X
+
+int XBlockSize; // The total number of elements in each block to be loaded
+
+int NumChannelBatches; // The number of iteration needed to cover all channels that needs to be processed. Equals to ceil(NumWChannels/NumChannelsPerBatch)
+int NumChannelsPerBatch; // The number of complete kernel channels that can be loaded in one batch
+
+int WOutputKernelStride; // The total number of W elements per output kernel and thus PROD(C/GROUPS, WD0, .., WDi)
+int WChannelBatchSize; // The total number of W elements inside a channel batch and thus the number of channels per batch times PROD(WD0, .., WDi)
+int WChannelSize; // Number of elements in a single channel of W and thus PROD(WDi)
+
+float GroupsDivM; // The number of groups divided by the number of output kernels used to compute the start channel offset
+
+groupshared WORK_TYPE SharedMemoryX[NUM_GROUP_THREADS<<NUM_READS_PER_THREAD_POW2];
+groupshared WORK_TYPE SharedMemoryW[NUM_GROUP_THREADS];
+
+// Must correspond to EConvAlgorithm defined in ConvCS.h
+#if ALGORITHM < 1
+
+[numthreads(NUM_GROUP_THREADS, 1, 1)]
+void main(in const uint3 DispatchThreadID : SV_DispatchThreadID, in const uint3 GroupID : SV_GroupID, in const uint3 GroupThreadID : SV_GroupThreadID)
+{
+	// Group Z dimension and Y dimension are used for batch and output kernel
+    // This works as each group only works on the inner dimensions (num threads in y and z set to one, group count to N and M respectively)
+	int BatchIndex = GroupID.z;
+	int OutputKernelIndex = GroupID.y;
+
+	// Input channels and output kernels are divided into groups
+	// Thus the channel offset for this thread's output kernel can be computed as follows:
+	// The output kernel index is divided by the number of output kernels per group giving the group index
+	// This is multiplied by the number of channels in W to get the start index
+	int InputChannelStartIndex = NumWChannels * (int)(GroupsDivM * (float)OutputKernelIndex);
+
+	// Compute the output memory index and dimension indices of this thread inside its group as well as the start dimension indices of the X block to load
+	int YIndex = BatchIndex * YBatchStride + OutputKernelIndex * YOutputKernelStride;
+	int GroupThreadDimensionIndex[NUM_DIMENSIONS];// This will be used later to lookup the shared W and X to compute the final result
+    int XBlockStartDimensionIndex[NUM_DIMENSIONS];
+	int GroupIndex = GroupID.x;
+	int GroupThreadIndex = GroupThreadID.x;
+	UNROLL
+	for (int i = 0; i < NUM_DIMENSIONS; i++)
+	{
+        // Divide by the number of groups contained by one index of this dimension to get the dimension index n being the nth group in this dimension
+		int GroupDimensionIndex = (int)(OneDiv_GroupStride_GroupThreadStride_XBlockStride[i].x * (float)GroupIndex);
+		GroupIndex -= GroupDimensionIndex * GroupStride_GroupShape_GroupThreadStride_StrideXBlockStride[i].x;
+
+        // Same here with number of threads contained inside the group giving the nth thread inside the group in this dimension
+		GroupThreadDimensionIndex[i] = (int)(OneDiv_GroupStride_GroupThreadStride_XBlockStride[i].y * (float)GroupThreadIndex);
+		GroupThreadIndex -= GroupThreadDimensionIndex[i] * GroupStride_GroupShape_GroupThreadStride_StrideXBlockStride[i].z;
+		
+        // Compose the global index by the group index times group shape plus thread index inside the group
+		int ThreadDimensionIndex = GroupDimensionIndex * GroupStride_GroupShape_GroupThreadStride_StrideXBlockStride[i].y + GroupThreadDimensionIndex[i];
+        
+        // Subtract the padding from the index of the first element of the group to get the start index of the block needed by this group
+        XBlockStartDimensionIndex[i] = GroupDimensionIndex * XBlockStartStride_XBlockStride_WDimension_WDimensionDilationXBlockStride[i].x + Dilation_Stride_XBlockStartOffset_DilationXBlockStride[i].z;
+
+        // ThreadDimensionIndex is always positive but needs to be tested if it does not overflow
+		if (ThreadDimensionIndex >= YDimension_YMemoryStride_XDimension_XMemoryStride[i].x || YIndex < 0)
+		{
+			YIndex = -1;
+		}
+		else
+		{
+			YIndex += ThreadDimensionIndex * YDimension_YMemoryStride_XDimension_XMemoryStride[i].y;
+		}
+	}
+
+    // Compute the memory indices in global X which the current thread has to load as well as the indices in shared X where to write
+	int GlobalXIndex[0x1<<NUM_READS_PER_THREAD_POW2];
+	int SharedXIndex[0x1<<NUM_READS_PER_THREAD_POW2];
+	int ReadIdx = GroupThreadID.x;
+    UNROLL
+    for (int j = 0; j < (int)(0x1<<NUM_READS_PER_THREAD_POW2); j++)
+	{
+        // If the read index exceeded the number of elements in the X block further computation can be skipped
+        if (ReadIdx >= XBlockSize)
+        {
+            // Assign -1 for both to indicate that the element is not used and thus neither has to be read from global nor written to shared memory
+            GlobalXIndex[j] = -1;
+            SharedXIndex[j] = -1;
+            continue;
+        }
+
+        // Turn the flat group index multiple into thread dimension indices inside the volume to be loaded by this group
+        int TmpGlobalXIndex = BatchIndex * XBatchStride + InputChannelStartIndex * XChannelStride;
+        int TmpSharedXIndex = 0;
+        int TempReadIdx = ReadIdx;
+        UNROLL
+        for (int k = 0; k < NUM_DIMENSIONS; k++)
+        {
+            int DimensionIndex = (int)(OneDiv_GroupStride_GroupThreadStride_XBlockStride[k].z * (float)TempReadIdx);
+            int DimensionIndexOffset = DimensionIndex * XBlockStartStride_XBlockStride_WDimension_WDimensionDilationXBlockStride[k].y;
+		    TempReadIdx -= DimensionIndexOffset;
+
+            // Update TmpGlobalXIndex only if positive (it is already outside the volume if negative)
+            if (TmpGlobalXIndex >= 0)
+            {
+                // Check if the global X index is in range
+                int XIndex = XBlockStartDimensionIndex[k] + DimensionIndex;
+                if (XIndex < 0 || XIndex >= YDimension_YMemoryStride_XDimension_XMemoryStride[k].z)
+                {
+                    TmpGlobalXIndex = -1;
+                }
+                else
+                {
+                    TmpGlobalXIndex += XIndex * YDimension_YMemoryStride_XDimension_XMemoryStride[k].w;
+                }
+            }
+
+            // Update TmpSharedXIndex in any case as a zero has to be written
+            TmpSharedXIndex += DimensionIndexOffset;
+        }
+
+        // Assign the result, which will be -1 for the global and positive for the shared index if the element lies in the padding area
+        // This indicates that the global element should not be loaded but a zero has to be written to shared memory 
+        GlobalXIndex[j] = TmpGlobalXIndex;
+        SharedXIndex[j] = TmpSharedXIndex;
+
+        // Increment the read idx by the number of threads inside the group to cover all elements
+        ReadIdx += NUM_GROUP_THREADS;
+    }
+
+    // Write the zeros once to avoid rewriting the same value later
+    UNROLL
+    for (int l = 0; l < (int)(0x1<<NUM_READS_PER_THREAD_POW2); l++)
+	{
+        // If the global index is out of bound (padding area) but the shared index is valid (inside the X block) the shared element needs initialization to zero
+        if (GlobalXIndex[l] < 0 && SharedXIndex[l] >= 0)
+        {
+            SharedMemoryX[SharedXIndex[l]] = 0;
+        }
+    }
+
+    // Iterate over all channel batches
+    int WChannelOutputKernelOffset = OutputKernelIndex * WOutputKernelStride;
+    int WChannelBatchOffset = 0;
+    int WChannelIndex = 0;
+#if HAS_B == 0
+    WORK_TYPE Result = 0;
+#else
+    WORK_TYPE Result = BUFFER_TO_WORK_TYPE(B[OutputKernelIndex]);
+#endif
+    for (int m = 0; m < NumChannelBatches; m++)
+	{
+        // Load all channels of W belonging to this batch
+        int WIndex = WChannelBatchOffset + GroupThreadID.x;
+        if (GroupThreadID.x < WChannelBatchSize && WIndex < WOutputKernelStride)
+        {
+            SharedMemoryW[GroupThreadID.x] = BUFFER_TO_WORK_TYPE(W[WChannelOutputKernelOffset + WIndex]);
+            WChannelBatchOffset += WChannelBatchSize;
+        }
+        
+        // Each channel batch iterates over all consecutive elements of the loaded kernel channels
+        // This results in bank conflict free lookups, as all threads of the same warp read exactly the same element which is then broadcasted
+        int SharedWIndex = 0;
+
+        // Iterate over all channels in this batch
+        for (int n = 0; n < NumChannelsPerBatch; n++)
+	    {
+            // Check if the channel index is valid as channel batches may not be aligned with the effective number of channels
+            if (WChannelIndex >= NumWChannels)
+            {
+                break;
+            }
+            WChannelIndex++;
+            
+            // Load the shared X
+            UNROLL
+            for (int o = 0; o < (int)(0x1<<NUM_READS_PER_THREAD_POW2); o++)
+	        {
+                // If the global index is negative, a zero has already been written, if shared index is negative there is no need to load anything
+                if (GlobalXIndex[o] >= 0 && SharedXIndex[o] >= 0)
+                {
+                    // Load the global data into shared memory
+                    SharedMemoryX[SharedXIndex[o]] = BUFFER_TO_WORK_TYPE(X[GlobalXIndex[o]]);
+
+                    // Update the indices for the next iteration by adding the channel stride
+                    GlobalXIndex[o] += XChannelStride;
+                }
+            }
+
+            // Sync the group before doing any computation
+            GroupMemoryBarrierWithGroupSync();
+
+            // Initialize a dimension index array used to iterate through the kernel dimensions
+            // Also initialize the start index in shared memory based on the result position
+            int WDimensionIndex[NUM_DIMENSIONS];
+            int SharedXIndex = 0;
+            UNROLL
+            for (int p = 0; p < NUM_DIMENSIONS; p++)
+            {
+                WDimensionIndex[p] = 0;
+
+                // The index of the first element (e.g. upper left) required by each output is offsetted by the kernel stride and then multiplied by the stride in memory
+                SharedXIndex += GroupThreadDimensionIndex[p] * GroupStride_GroupShape_GroupThreadStride_StrideXBlockStride[p].w;
+            }
+
+            // Compute and accumulate the result
+            for (int q = 0; q < WChannelSize; q++)
+            {
+                Result += SharedMemoryX[SharedXIndex] * SharedMemoryW[SharedWIndex];
+
+                // Advance the inner-most dimension
+                SharedWIndex++;
+                SharedXIndex += Dilation_Stride_XBlockStartOffset_DilationXBlockStride[NUM_DIMENSIONS-1].x;
+                WDimensionIndex[NUM_DIMENSIONS-1]++;
+
+                // Advance the other dimensions if necessary
+                UNROLL
+                for (int r = NUM_DIMENSIONS-1; r > 0; r--)
+                {
+                    if(WDimensionIndex[r] >= XBlockStartStride_XBlockStride_WDimension_WDimensionDilationXBlockStride[r].z)
+                    {
+                        // Adjust the SharedXIndex by resetting the current dimension and increase the next outer one
+                        SharedXIndex -= XBlockStartStride_XBlockStride_WDimension_WDimensionDilationXBlockStride[r].w;
+                        SharedXIndex += Dilation_Stride_XBlockStartOffset_DilationXBlockStride[r-1].w;
+                        WDimensionIndex[r] = 0;
+                        WDimensionIndex[r-1]++;
+                    }
+                    else 
+                    {
+                        break;
+                    }
+                }
+            }
+
+            // Sync the group before continuing any load
+            GroupMemoryBarrierWithGroupSync();
+        }
+    }
+
+    // Write the final result
+    if (YIndex >= 0)
+    {
+        Y[YIndex] = WORK_TO_BUFFER_TYPE(Result);
+    }
+}
+
+#endif
\ No newline at end of file
