NNX Parametric tests proof of concept
* json parser + schema
* run test via command or automation (with filtering)
* models test + 2 operators as a demo of the feature

diff --git a/Engine/Plugins/Experimental/NNX/Source/NNXQA/NNXQA.Build.cs b/Engine/Plugins/Experimental/NNX/Source/NNXQA/NNXQA.Build.cs
--- a/Engine/Plugins/Experimental/NNX/Source/NNXQA/NNXQA.Build.cs
+++ b/Engine/Plugins/Experimental/NNX/Source/NNXQA/NNXQA.Build.cs
@@ -18,6 +18,8 @@ public class NNXQA : ModuleRules
 			new string[]
 			{
 				"Engine",
+				"Projects",
+				"Json",
 				"CoreUObject",
 				"NNXCore",
 				"NNXUtils"
diff --git a/Engine/Plugins/Experimental/NNX/Source/NNXQA/Private/NNXQAJsonUtils.h b/Engine/Plugins/Experimental/NNX/Source/NNXQA/Private/NNXQAJsonUtils.h
new file mode 100644
--- /dev/null
+++ b/Engine/Plugins/Experimental/NNX/Source/NNXQA/Private/NNXQAJsonUtils.h
@@ -0,0 +1,178 @@
+// Copyright Epic Games, Inc. All Rights Reserved.
+
+#pragma once
+
+#include "Serialization/JsonReader.h"
+#include "Serialization/JsonSerializer.h"
+#include "Serialization/JsonSerializerMacros.h"
+#include "Dom/JsonObject.h"
+#include "Misc/FileHelper.h"
+
+namespace NNX 
+{
+namespace Test
+{
+namespace Json
+{
+	struct FTestConfigTensor : FJsonSerializable
+	{
+		TArray<int32> Shape;
+		//TODO Extend:
+		// by adding custom data
+		// by adding custom initializer or parameter (range, distribution)
+		// by adding data type
+		// by adding source CPU or GPU resident data
+
+		BEGIN_JSON_SERIALIZER
+			JSON_SERIALIZE_ARRAY("shape", Shape);
+			//JSON_SERIALIZE("type", Type);
+			//JSON_SERIALIZE("on_gpu", OnGpu);
+			//JSON_SERIALIZE("range", ...);
+			//JSON_SERIALIZE("data_source", DataSource);
+		END_JSON_SERIALIZER
+	};
+
+	struct FTestConfigRuntime : FJsonSerializable
+	{
+		FString Name;
+		bool Skip;
+		float AbsoluteError;
+		float RelativeError;
+
+		BEGIN_JSON_SERIALIZER
+			JSON_SERIALIZE("name", Name);
+			JSON_SERIALIZE("skip", Skip);
+			JSON_SERIALIZE("absolute_error", AbsoluteError);
+			JSON_SERIALIZE("relative_error", RelativeError);
+		END_JSON_SERIALIZER
+	};
+
+	struct FTestConfigDataset : FJsonSerializable
+	{
+		TArray<FTestConfigTensor> Inputs;
+		TArray<FTestConfigTensor> Outputs;
+		TArray<FTestConfigRuntime> Runtimes;
+
+		BEGIN_JSON_SERIALIZER
+			JSON_SERIALIZE_ARRAY_SERIALIZABLE("inputs", Inputs, FTestConfigTensor);
+			JSON_SERIALIZE_ARRAY_SERIALIZABLE("outputs", Outputs, FTestConfigTensor);
+			JSON_SERIALIZE_ARRAY_SERIALIZABLE("runtimes", Runtimes, FTestConfigRuntime);
+		END_JSON_SERIALIZER
+	};
+
+	struct FTestConfigTarget : FJsonSerializable
+	{
+		FString Target;
+		TArray<FString> Tags;
+		TArray<FString> AdditionalDatasets;
+		bool Skip;
+		float AbsoluteError;
+		float RelativeError;
+		TArray<FTestConfigRuntime> Runtimes;
+
+		BEGIN_JSON_SERIALIZER
+			JSON_SERIALIZE("target", Target);
+			JSON_SERIALIZE("skip", Skip);
+			JSON_SERIALIZE("absolute_error", AbsoluteError);
+			JSON_SERIALIZE("relative_error", RelativeError);
+			JSON_SERIALIZE_ARRAY_SERIALIZABLE("runtimes", Runtimes, FTestConfigRuntime);
+			JSON_SERIALIZE_ARRAY("tags", Tags);
+			JSON_SERIALIZE_ARRAY("additional_datasets", AdditionalDatasets);
+		END_JSON_SERIALIZER
+	};
+
+	struct FTestCategory : FJsonSerializable
+	{
+		FString Category;
+		TArray<FTestConfigTarget> Targets;
+		bool Skip;
+		bool IsModelTest;
+
+		BEGIN_JSON_SERIALIZER
+			JSON_SERIALIZE("category", Category);
+			JSON_SERIALIZE("skip", Skip);
+			JSON_SERIALIZE("is_model_test", IsModelTest);
+			JSON_SERIALIZE_ARRAY_SERIALIZABLE("targets", Targets, FTestConfigTarget);
+		END_JSON_SERIALIZER
+	};
+
+	struct FTestConfigInputOutputSet : FJsonSerializable
+	{
+		FString Name;
+		TArray<FTestConfigDataset> Datasets;
+
+		BEGIN_JSON_SERIALIZER
+			JSON_SERIALIZE("name", Name);
+			JSON_SERIALIZE_ARRAY_SERIALIZABLE("datasets", Datasets, FTestConfigDataset);
+		END_JSON_SERIALIZER
+	};
+
+	bool LoadTestDescriptionFromJson(const FString& FullPath,
+		TArray<FTestCategory>& ModelTestCategories,
+		TArray<FTestCategory>& OperatorCategories,
+		TArray<FTestConfigInputOutputSet>& InputOutputSets)
+	{
+		ModelTestCategories.Empty();
+		OperatorCategories.Empty();
+		InputOutputSets.Empty();
+
+		FString JsonContent;
+		FFileHelper::LoadFileToString(JsonContent, *FullPath);
+		TSharedRef<TJsonReader<TCHAR>> JsonReader = TJsonReaderFactory<TCHAR>::Create(JsonContent);
+		TSharedPtr<FJsonObject> GlobalJsonObject = MakeShareable(new FJsonObject);
+
+		if (!FJsonSerializer::Deserialize(JsonReader, GlobalJsonObject))
+		{
+			return false;
+		}
+
+		if (!GlobalJsonObject.IsValid())
+		{
+			return false;
+		}
+
+		const TArray<TSharedPtr<FJsonValue>>* JsonArrayValue;
+		if (GlobalJsonObject->TryGetArrayField(TEXT("model_test_categories"), JsonArrayValue))
+		{
+			for (TSharedPtr<FJsonValue> JsonValue : *JsonArrayValue)
+			{
+				const TSharedPtr<FJsonObject>* JsonObject = nullptr;
+				if (!JsonValue->TryGetObject(JsonObject))
+				{
+					return false;
+				}
+				ModelTestCategories.Emplace_GetRef().FromJson(*JsonObject);
+			}
+		}
+
+		if (GlobalJsonObject->TryGetArrayField(TEXT("operator_test_categories"), JsonArrayValue))
+		{
+			for (TSharedPtr<FJsonValue> JsonValue : *JsonArrayValue)
+			{
+				const TSharedPtr<FJsonObject>* JsonObject = nullptr;
+				if (!JsonValue->TryGetObject(JsonObject))
+				{
+					return false;
+				}
+				OperatorCategories.Emplace_GetRef().FromJson(*JsonObject);
+			}
+		}
+
+		if (GlobalJsonObject->TryGetArrayField(TEXT("input_output_sets"), JsonArrayValue))
+		{
+			for (TSharedPtr<FJsonValue> JsonValue : *JsonArrayValue)
+			{
+				const TSharedPtr<FJsonObject>* JsonObject = nullptr;
+				if (!JsonValue->TryGetObject(JsonObject))
+				{
+					return false;
+				}
+				InputOutputSets.Emplace_GetRef().FromJson(*JsonObject);
+			}
+		}
+		return true;
+	}
+
+} // namespace Json
+} // namespace Test
+} // namespace NNX
diff --git a/Engine/Plugins/Experimental/NNX/Source/NNXQA/Private/NNXQAParametricTest.cpp b/Engine/Plugins/Experimental/NNX/Source/NNXQA/Private/NNXQAParametricTest.cpp
new file mode 100644
--- /dev/null
+++ b/Engine/Plugins/Experimental/NNX/Source/NNXQA/Private/NNXQAParametricTest.cpp
@@ -0,0 +1,443 @@
+// Copyright Epic Games, Inc. All Rights Reserved.
+
+#include "NNXQAParametricTest.h"
+#include "NNXCore.h"
+#include "NNXQAUtils.h"
+#include "NNXQAJsonUtils.h"
+#include "NNXModelBuilder.h"
+#include "Interfaces/IPluginManager.h"
+#include "Misc/Paths.h"
+#include "HAL/IConsoleManager.h"
+
+#if WITH_DEV_AUTOMATION_TESTS
+#include "Misc/AutomationTest.h"
+#endif
+
+namespace NNX 
+{
+namespace Test
+{
+	class FParametricTests : public FTests
+	{
+	public:
+		FParametricTests()
+		{
+			ReloadTestDescriptionsFromJson();
+		}
+
+		bool ReloadTestDescriptionsFromJson()
+		{
+			TestSetups.Empty();
+
+			TSharedPtr<IPlugin> NNXPlugin = IPluginManager::Get().FindPlugin(TEXT("NNX"));
+			
+			if (!NNXPlugin.IsValid())
+			{
+				return false;
+			}
+
+			//TODO verify the path can be accessed on standalone build and on consoles
+			//TODO allow to define the tests in more than one json file
+			FString NNXPluginBaseDir = NNXPlugin->GetBaseDir();
+			FString FullPath(NNXPluginBaseDir + TEXT("\\Source\\NNXQA\\Resources\\NNXQATestDesc.json"));
+
+			TArray<Json::FTestCategory> ModelTestCategories;
+			TArray<Json::FTestCategory> OperatorTestCategories;
+			TArray<Json::FTestConfigInputOutputSet> InputOutputSets;
+			
+			if (!Json::LoadTestDescriptionFromJson(FullPath, ModelTestCategories, OperatorTestCategories, InputOutputSets))
+			{
+				return false;
+			}
+
+			const FString NNXBaseTestPath(TEXT("System.Engine.MachineLearning.NNX.Parametric"));
+			AddTestFromCategory(NNXBaseTestPath + TEXT(".Model."), ModelTestCategories, InputOutputSets);
+			AddTestFromCategory(NNXBaseTestPath + TEXT(".Operator."), OperatorTestCategories, InputOutputSets);
+
+			return true;
+		}
+	private:
+		
+		static void ApplyEpsilons(FTests::FTestSetup& TestSetup, const Json::FTestConfigTarget& TestTarget)
+		{
+			if (TestTarget.AbsoluteError != 0.0f)
+			{
+				TestSetup.AbsoluteErrorEpsilon = TestTarget.AbsoluteError;
+			}
+			
+			if (TestTarget.RelativeError != 0.0f)
+			{
+				TestSetup.RelativeErrorPercent = TestTarget.RelativeError;
+			}
+		}
+
+		static void ApplyRuntimesConfig(FTests::FTestSetup& TestSetup, const TArray<Json::FTestConfigRuntime>& TestRuntimes)
+		{
+			for (const Json::FTestConfigRuntime& Runtime : TestRuntimes)
+			{
+				if (Runtime.Skip && !TestSetup.AutomationExcludedRuntime.Contains(Runtime.Name))
+				{
+					TestSetup.AutomationExcludedRuntime.Emplace(Runtime.Name);
+				}
+				else
+				{
+					if (Runtime.AbsoluteError != 0.0f)
+					{
+						TestSetup.AbsoluteErrorEpsilonForRuntime.Emplace(Runtime.Name, Runtime.AbsoluteError);
+					}
+					if (Runtime.RelativeError != 0.0f)
+					{
+						TestSetup.RelativeErrorPercentForRuntime.Emplace(Runtime.Name, Runtime.RelativeError);
+					}
+				}
+			}
+		}
+
+		static void ApplyTargetConfig(FTests::FTestSetup& TestSetup, const Json::FTestConfigTarget& TestTarget)
+		{
+			//Epsilons
+			if (TestTarget.AbsoluteError != 0.0f)
+			{
+				TestSetup.AbsoluteErrorEpsilon = TestTarget.AbsoluteError;
+			}
+			if (TestTarget.RelativeError != 0.0f)
+			{
+				TestSetup.RelativeErrorPercent = TestTarget.RelativeError;
+			}
+
+			ApplyRuntimesConfig(TestSetup, TestTarget.Runtimes);;
+
+			//Tags
+			TestSetup.Tags = TestTarget.Tags;
+		}
+
+		static TArray<uint32> GetShapeFromJsonArray(const TArray<int32>& JsonShape)
+		{
+			TArray<uint32> Shape;
+			Shape.Reserve(JsonShape.Num());
+			for (int32 dim : JsonShape)
+			{
+				Shape.Emplace(dim >= 0 ? dim : 1);
+			}
+
+			return Shape;
+		}
+
+		static void ApplyDatasetConfig(FTests::FTestSetup& TestSetup, const Json::FTestConfigDataset& TestDataset)
+		{
+			ApplyRuntimesConfig(TestSetup, TestDataset.Runtimes);
+
+			if (TestDataset.Inputs.Num() == 0)
+			{
+				return;
+			}
+
+			//TODO allow more tensor type
+			EMLTensorDataType TensorType = EMLTensorDataType::Float;
+
+			uint32 i = 0;
+			for (auto&& Tensor : TestDataset.Inputs)
+			{
+				TArray<uint32> shape = GetShapeFromJsonArray(Tensor.Shape);
+				TestSetup.Inputs.Emplace(FMLTensorDesc::Make(FString::Printf(TEXT("in%d"), i++), shape, TensorType));
+			}
+
+			i = 0;
+			for (auto&& Tensor : TestDataset.Outputs)
+			{
+				TArray<uint32> shape = GetShapeFromJsonArray(Tensor.Shape);
+				TestSetup.Outputs.Emplace(FMLTensorDesc::Make(FString::Printf(TEXT("output%d"), i++), shape, TensorType));
+			}
+			//If output is not defined it is the first input shape.
+			if (TestDataset.Outputs.Num() == 0 && TestDataset.Inputs.Num() > 0)
+			{
+				TArray<uint32> shape = GetShapeFromJsonArray(TestDataset.Inputs[0].Shape);
+				TestSetup.Outputs.Emplace(FMLTensorDesc::Make(TEXT("output"), shape, TensorType));
+			}
+		}
+
+		static FString GetTestSuffix(const Json::FTestConfigDataset& Dataset)
+		{
+			//Build TestSuffix "<inputshape0>_...=><outputshape0>_..."
+			bool bIsFirstShape = true;
+			FString TestSuffix;
+			for (const Json::FTestConfigTensor& Input : Dataset.Inputs)
+			{
+				if (!bIsFirstShape) TestSuffix += TEXT("_");
+				TestSuffix += ShapeToString(Input.Shape);
+				bIsFirstShape = false;
+			}
+			TestSuffix += TEXT("=>");
+			//If output is not defined it is the first input shape.
+			if (Dataset.Outputs.Num() == 0)
+			{
+				TestSuffix += ShapeToString(Dataset.Inputs[0].Shape);
+			}
+			else
+			{
+				for (const Json::FTestConfigTensor& Output : Dataset.Outputs)
+				{
+					TestSuffix += TEXT("_");
+					TestSuffix += ShapeToString(Output.Shape);
+				}
+			}
+			return TestSuffix;
+		}
+		
+		void AddTestFromCategory(const FString& BaseTestPath, const TArray<Json::FTestCategory>& TestCategories, const TArray<Json::FTestConfigInputOutputSet> InputOutputSets)
+		{
+			for (const Json::FTestCategory& TestCategory : TestCategories)
+			{
+				if (TestCategory.Skip)
+				{
+					continue;
+				}
+
+				const bool bIsModelCategory = TestCategory.IsModelTest;
+				const FString TestCategoryPath(BaseTestPath + TestCategory.Category + TEXT("."));
+				for (const Json::FTestConfigTarget& TestTarget : TestCategory.Targets)
+				{
+					if (TestTarget.Skip)
+					{
+						continue;
+					}
+
+					const FString& TestBaseName = TestTarget.Target;
+
+					bool bAtLeastATestWasAdded = false;
+					for (const Json::FTestConfigInputOutputSet& InputOutputSet : InputOutputSets)
+					{
+						//If category is a substring of InputOutputSet name or if target explicitly requested the InputOutputSet name
+						if (InputOutputSet.Name.Contains(TestCategory.Category) || TestTarget.AdditionalDatasets.Contains(InputOutputSet.Name))
+						{
+							for (const Json::FTestConfigDataset& Dataset : InputOutputSet.Datasets)
+							{
+								if (Dataset.Inputs.Num() == 0)
+								{
+									continue;
+								}
+
+								FTests::FTestSetup& Test = AddTest(TestCategoryPath, TestBaseName, TEXT(".") + GetTestSuffix(Dataset));
+								ApplyTargetConfig(Test, TestTarget);
+								ApplyDatasetConfig(Test, Dataset);
+								Test.IsModelTest = bIsModelCategory;
+							}
+							bAtLeastATestWasAdded = true;
+						}
+					}
+
+					//No dataset were matched with this target, define a test without enforcing shapes.
+					if (!bAtLeastATestWasAdded)
+					{
+						FTests::FTestSetup& Test = AddTest(TestCategoryPath, TestBaseName, TEXT(""));
+						ApplyTargetConfig(Test, TestTarget);
+						Test.IsModelTest = bIsModelCategory;
+					}
+				}
+			}
+		}
+	};
+	static FParametricTests ParametricTests;
+
+	static FString GetFullModelPath(const FString& ModelName)
+	{
+		//Note: This mean model tests can only run in the context of the current projects (example: NNXIncubator)
+		return FPaths::ConvertRelativePathToFull(FPaths::ProjectContentDir() / TEXT("OnnxModels") / ModelName);
+	}
+	
+	static bool RunParametricTest(FTests::FTestSetup& TestSetup, const FString& RuntimeFilter)
+	{
+		TArray<uint8> ModelData;
+
+		if (TestSetup.IsModelTest)
+		{
+			// Model test, load model from disk
+			FString ModelPath = GetFullModelPath(TestSetup.TargetName + TEXT(".onnx"));
+			const bool bIsModelInMem = FFileHelper::LoadFileToArray(ModelData, *ModelPath);
+			if (!bIsModelInMem)
+			{
+				UE_LOG(LogNNX, Error, TEXT("Can't load model from disk at path '%s'. Tests ABORTED!"), *ModelPath);
+				return false;
+			}
+		}
+		else
+		{
+			// Operator test, create model in memory
+			if (!CreateONNXModelForOperator(TestSetup.TargetName, TestSetup.Inputs, TestSetup.Outputs, ModelData))
+			{
+				UE_LOG(LogNNX, Error, TEXT("Failed to create model for test '%s'. Test ABORTED!"), *TestSetup.TargetName);
+				return false;
+			}
+		}
+
+		return CompareONNXModelInferenceAcrossRuntimes(TestSetup.TestName, ModelData, &TestSetup, RuntimeFilter);
+	}
+
+	static FString AutomationRuntimeFilter;
+	void SetAutomationRuntimeFilter(const FString& InRuntimeFilter)
+	{
+		AutomationRuntimeFilter = InRuntimeFilter;
+	}
+
+	static FAutoConsoleCommand SetAutomationRuntimeFilterCommand(
+		TEXT("nnx.test.parametric.setautomationfilter"), TEXT("Set the RuntimeFilter witch automation will use, no parameter to run on all runtime (this is the default)."),
+		FConsoleCommandWithArgsDelegate::CreateStatic(
+			[](const TArray< FString >& Args)
+			{
+				if (Args.Num() > 0)
+				{
+					SetAutomationRuntimeFilter(Args[0]);
+				}
+				else
+				{
+					SetAutomationRuntimeFilter(TEXT(""));
+				}
+			}
+		)
+	);
+	
+	bool RunParametricTests(const FString& NameSubstring, const FString& Tag, const FString& InRuntimeFilter)
+	{
+		bool bAllTestSucceed = true;
+		uint32 NumTest = 0;
+		uint32 NumTestFailed = 0;
+		for (FTests::FTestSetup& Test : ParametricTests.TestSetups)
+		{
+			if (!Tag.IsEmpty() && !Test.Tags.Contains(Tag))
+			{
+				continue;
+			}
+			if (!NameSubstring.IsEmpty() && !Test.TestName.Contains(NameSubstring))
+			{
+				continue;
+			}
+
+			++NumTest;
+			if (!RunParametricTest(Test, InRuntimeFilter))
+			{
+				++NumTestFailed;
+			}
+		}
+
+		if (NumTest == 0)
+		{
+			UE_LOG(LogNNX, Display, TEXT("No test selected to run (on %d parametric tests registered)."), ParametricTests.TestSetups.Num());
+			return true;
+		}
+		else if (NumTestFailed == 0)
+		{
+			UE_LOG(LogNNX, Display, TEXT("SUCCEED! All %d tests selected passed (%d tests are registered)."), NumTest, ParametricTests.TestSetups.Num());
+			return true;
+		}
+		else
+		{
+			UE_LOG(LogNNX, Error, TEXT("FAILED! %d test(s) failed, on the %d test selected to run (%d tests are registered)."), NumTestFailed, NumTest, ParametricTests.TestSetups.Num());
+			return false;
+		}
+	}
+	
+	static FString FindArg(const FString& ArgName, const TArray< FString >& Args)
+	{
+		FString Arg(TEXT(""));
+		int32 ArgNameIndex = Args.Find(ArgName);
+		if (ArgNameIndex == INDEX_NONE || ArgNameIndex + 1 >= Args.Num())
+		{
+			return Arg;
+		}
+		else
+		{
+			return Args[ArgNameIndex + 1];
+		}
+	}
+
+	static FAutoConsoleCommand RunTestCommand(
+		TEXT("nnx.test.parametric.run"), TEXT("Run all parametric tests. Use -name to filter by name (substring). Use -tag to filter by tag. Use -runtime to only run for the provided runtime."),
+		FConsoleCommandWithArgsDelegate::CreateStatic(
+			[](const TArray< FString >& Args)
+			{
+				FString Name = FindArg(TEXT("-name"), Args);
+				FString Tag = FindArg(TEXT("-tag"), Args);
+				FString Runtime = FindArg(TEXT("-runtime"), Args);
+				RunParametricTests(Name, Tag, Runtime);
+			}
+		)
+	);
+
+#if WITH_DEV_AUTOMATION_TESTS
+	
+	IMPLEMENT_SIMPLE_AUTOMATION_TEST_PRIVATE(FNNXParametricTestBase, FAutomationTestBase, "NNXParametricTest", EAutomationTestFlags::ApplicationContextMask | EAutomationTestFlags::FeatureMask | EAutomationTestFlags::EngineFilter, __FILE__, __LINE__)
+	bool FNNXParametricTestBase::RunTest(const FString& Parameters) { return false; }
+
+	class FNNXParametricTest : public FNNXParametricTestBase
+	{
+		FTests::FTestSetup Test;
+
+	public:
+		FNNXParametricTest(const FTests::FTestSetup& InTest) : FNNXParametricTestBase(InTest.TestName), Test(InTest) {}
+		virtual ~FNNXParametricTest() {}
+		virtual FString GetTestSourceFileName() const override { return "From Json"; }//TODO return source json file path
+		virtual int32 GetTestSourceFileLine() const override { return 0; }
+	protected:
+		virtual FString GetBeautifiedTestName() const override { return Test.TestName; }
+		bool RunTest(const FString& Parameter)
+		{
+			return RunParametricTest(Test, AutomationRuntimeFilter);
+		}
+	};
+
+	class FParametricTestAutomationRegistry
+	{
+	
+		TArray<FNNXParametricTest*> RegisteredTests;
+
+	public:
+		FParametricTestAutomationRegistry()
+		{
+			Refresh();
+		}
+		
+		~FParametricTestAutomationRegistry()
+		{
+			Clear();
+		}
+
+		void Clear()
+		{
+			for (FNNXParametricTest* Test : RegisteredTests)
+			{
+				delete Test;
+			}
+			RegisteredTests.Empty();
+		}
+
+		void Refresh()
+		{
+			Clear();
+
+			for (const FTests::FTestSetup& Test : ParametricTests.TestSetups)
+			{
+				RegisteredTests.Emplace(new FNNXParametricTest(Test));
+			}
+		}
+	};
+
+	static FParametricTestAutomationRegistry ParametricTestAutomationRegistry;
+
+#endif //WITH_DEV_AUTOMATION_TESTS
+
+	static FAutoConsoleCommand TestReloadCommand(
+		TEXT("nnx.test.parametric.reload"), TEXT("Reload NNX parametric tests definition from Json."),
+		FConsoleCommandWithArgsDelegate::CreateStatic(
+			[](const TArray< FString >& Args)
+			{
+				ParametricTests.ReloadTestDescriptionsFromJson();
+				
+				#if WITH_DEV_AUTOMATION_TESTS
+					ParametricTestAutomationRegistry.Refresh();
+				#endif
+			}
+		)
+	);
+
+} // namespace Test
+} // namespace NNX
diff --git a/Engine/Plugins/Experimental/NNX/Source/NNXQA/Private/NNXQAUtils.cpp b/Engine/Plugins/Experimental/NNX/Source/NNXQA/Private/NNXQAUtils.cpp
--- a/Engine/Plugins/Experimental/NNX/Source/NNXQA/Private/NNXQAUtils.cpp
+++ b/Engine/Plugins/Experimental/NNX/Source/NNXQA/Private/NNXQAUtils.cpp
@@ -83,6 +83,20 @@ namespace Test
 		return TestSuffix;
 	}
 
+	FString ShapeToString(TArrayView<const int32> Shape)
+	{
+		FString TestSuffix(TEXT("["));
+		bool bIsFirstDim = true;
+		for (uint32 size : Shape)
+		{
+			if (!bIsFirstDim) TestSuffix += TEXT(",");
+			TestSuffix += FString::Printf(TEXT("%d"), size);
+			bIsFirstDim = false;
+		}
+		TestSuffix += TEXT("]");
+		return TestSuffix;
+	}
+
 	FString FMLTensorDescToString(const FMLTensorDesc& desc)
 	{
 		TArrayView<const uint32> Shape(desc.Sizes, desc.Dimension);
@@ -362,7 +376,7 @@ namespace Test
 		return ReturnValue;
 	}
 
-	bool CompareONNXModelInferenceAcrossRuntimes(const FString& Name, TArrayView<uint8> ONNXModelData, const FTests::FTestSetup* AutomationTestSetup)
+	bool CompareONNXModelInferenceAcrossRuntimes(const FString& Name, TArrayView<uint8> ONNXModelData, const FTests::FTestSetup* AutomationTestSetup, const FString& RuntimeFilter)
 	{
 		FString CurrentPlatform = UGameplayStatics::GetPlatformName();
 		if (AutomationTestSetup != nullptr && AutomationTestSetup->AutomationExcludedPlatform.Contains(CurrentPlatform))
@@ -394,7 +408,14 @@ namespace Test
 		{
 			const FString& RuntimeName = Runtime->GetRuntimeName();
 			if (RuntimeName == RefName)
+			{
+				continue;
+			}
+
+			if (!RuntimeFilter.IsEmpty() && RuntimeName != RuntimeFilter)
+			{
 				continue;
+			}
 
 			FString TestResult;
 
@@ -437,15 +458,20 @@ namespace Test
 		}
 		return bAllTestsSucceeded;
 	}
-
-	FTests::FTestSetup& FTests::AddTest(const FString& ModelOrOperatorName, const FString& TestSuffix)
+	FTests::FTestSetup& FTests::AddTest(const FString & Category, const FString & ModelOrOperatorName, const FString & TestSuffix)
 	{
-		FString TestName = ModelOrOperatorName + TestSuffix;
+		FString TestName = Category + ModelOrOperatorName + TestSuffix;
 		//Test name should be unique
 		check(nullptr == TestSetups.FindByPredicate([TestName](const FTestSetup& Other) { return Other.TestName == TestName; }));
-		TestSetups.Emplace(ModelOrOperatorName, TestSuffix);
+		TestSetups.Emplace(Category, ModelOrOperatorName, TestSuffix);
 		return TestSetups.Last(0);
 	}
 
+	FTests::FTestSetup& FTests::AddTest(const FString& ModelOrOperatorName, const FString& TestSuffix)
+	{
+		//TODO cleanup: always use the 3 parameter path.
+		return AddTest(TEXT(""), ModelOrOperatorName, TestSuffix);
+	}
+
 } // namespace Test
 } // namespace NNX
\ No newline at end of file
diff --git a/Engine/Plugins/Experimental/NNX/Source/NNXQA/Private/NNXQAUtils.h b/Engine/Plugins/Experimental/NNX/Source/NNXQA/Private/NNXQAUtils.h
--- a/Engine/Plugins/Experimental/NNX/Source/NNXQA/Private/NNXQAUtils.h
+++ b/Engine/Plugins/Experimental/NNX/Source/NNXQA/Private/NNXQAUtils.h
@@ -14,20 +14,35 @@ namespace Test
 	{
 		struct FTestSetup
 		{
+			static constexpr float DefaultAbsoluteErrorEpsilon = 1e-5f;
+			static constexpr float DefaultRelativeErrorPercent = 1e-3f;
+
+			FTestSetup(const FString& InTestCategory, const FString& InModelOrOperatorName, const FString& InTestSuffix) :
+				TestName(InTestCategory + InModelOrOperatorName + InTestSuffix),
+				TargetName(InModelOrOperatorName),
+				AbsoluteErrorEpsilon(DefaultAbsoluteErrorEpsilon),
+				RelativeErrorPercent(DefaultRelativeErrorPercent),
+				IsModelTest(false)
+			{}
+			
 			FTestSetup(const FString& ModelOrOperatorName, const FString& TestSuffix) :
 				TestName(ModelOrOperatorName+TestSuffix),
 				TargetName(ModelOrOperatorName),
-				AbsoluteErrorEpsilon(1e-5f),
-				RelativeErrorPercent(1e-3f) {}
+				AbsoluteErrorEpsilon(DefaultAbsoluteErrorEpsilon),
+				RelativeErrorPercent(DefaultRelativeErrorPercent),
+				IsModelTest(false)
+			{}
 
 			const FString TestName;
 			const FString TargetName;
 			float AbsoluteErrorEpsilon;
 			float RelativeErrorPercent;
+			bool IsModelTest;
 			TMap<FString, float> AbsoluteErrorEpsilonForRuntime;
 			TMap<FString, float> RelativeErrorPercentForRuntime;
 			TArray<FMLTensorDesc> Inputs;
 			TArray<FMLTensorDesc> Outputs;
+			TArray<FString> Tags;
 			TArray<FString> AutomationExcludedRuntime;
 			TArray<FString> AutomationExcludedPlatform;
 			TArray<TPair<FString, FString>> AutomationExcludedPlatformRuntimeCombination;
@@ -46,16 +61,18 @@ namespace Test
 		};
 
 
+		FTestSetup& AddTest(const FString& Category, const FString& ModelOrOperatorName, const FString& TestSuffix);
 		FTestSetup& AddTest(const FString& ModelOrOperatorName, const FString& TestSuffix);
 
 		TArray<FTestSetup> TestSetups;
 	};
 
-	bool CompareONNXModelInferenceAcrossRuntimes(const FString& Name, TArrayView<uint8> ModelData, const FTests::FTestSetup* AutomationTestSetup);
+	bool CompareONNXModelInferenceAcrossRuntimes(const FString& Name, TArrayView<uint8> ModelData, const FTests::FTestSetup* AutomationTestSetup, const FString& RuntimeFilter = TEXT(""));
 
 	FString FMLTensorDescToString(const FMLTensorDesc& TensorDesc);
 	FString TensorToString(const FMLTensorDesc& TensorDesc, const TArray<char>& TensorData);
 	FString ShapeToString(TArrayView<const uint32> Shape);
+	FString ShapeToString(TArrayView<const int32> Shape);
 
 } // namespace Test
 } // namespace NNX
\ No newline at end of file
diff --git a/Engine/Plugins/Experimental/NNX/Source/NNXQA/Public/NNXQAParametricTest.h b/Engine/Plugins/Experimental/NNX/Source/NNXQA/Public/NNXQAParametricTest.h
new file mode 100644
--- /dev/null
+++ b/Engine/Plugins/Experimental/NNX/Source/NNXQA/Public/NNXQAParametricTest.h
@@ -0,0 +1,21 @@
+// Copyright Epic Games, Inc. All Rights Reserved.
+
+#pragma once
+
+#include "CoreMinimal.h"
+
+namespace NNX 
+{
+namespace Test 
+{
+	// Run all parametric tests 
+	// * NameSubstring will filtering test by full name (NameSubstring should be contained in the name). Empty string to run all tests.
+	// * Tag will filter tests by tag. Empty string to run all tests.
+	// * RuntimeFilter all to run those test only on the provided runtime. Empty string to run on all runtime.
+	NNXQA_API bool RunParametricTests(const FString& NameSubstring, const FString& Tag, const FString& RuntimeFilter);
+
+	//Set the RuntimeFilter witch automation will use, Empty string to run on all runtime (this is default)
+	NNXQA_API void SetAutomationRuntimeFilter(const FString& RuntimeFilter);
+
+} // namespace Test
+} // namespace NNX
\ No newline at end of file
diff --git a/Engine/Plugins/Experimental/NNX/Source/NNXQA/Resources/NNXQATestDesc.json b/Engine/Plugins/Experimental/NNX/Source/NNXQA/Resources/NNXQATestDesc.json
new file mode 100644
--- /dev/null
+++ b/Engine/Plugins/Experimental/NNX/Source/NNXQA/Resources/NNXQATestDesc.json
@@ -0,0 +1,124 @@
+{
+	"$schema": "./NNXQATestDescSchema.json",
+	"model_test_categories": [
+		{
+			"category": "OnnxZoo",
+			"is_model_test": true,
+			"targets": [
+				{
+					"target": "resnet50-v2-7",
+					"tags": [ "SmokeTest" ],
+					"runtimes": [
+						{ "name": "NNXRuntimeORTDml", "relative_error": 1.0 },
+						{ "name": "NNXRuntimeHlsl",	"skip": true, "comment": "Model with multiple layer not yet implemented."	},
+						{ "name": "NNXRuntimeDml",  "skip": true, "comment": "Model with multiple layer not yet implemented." }
+					]
+				},
+				{
+					"target": "Mosaic-9",
+					"tags": [ "SmokeTest" ],
+					"runtimes": [
+						{ "name": "NNXRuntimeORTDml", "absolute_error": 5e-2, "relative_error": 50.0,
+							"comment": "//Worst relative error 43.367836 % (epsilon 0.001000 %) at position 112558, got -0.003391 expected -0.002365. TODO we need an automatic balance between relative error and absolute error"
+						},
+						{ "name": "NNXRuntimeCPU", "absolute_error": 1e-3, "relative_error": 3.0 },
+						{ "name": "NNXRuntimeHlsl",	"skip": true, "comment": "Model with multiple layer not yet implemented."	},
+						{ "name": "NNXRuntimeDml",  "skip": true, "comment": "Model with multiple layer not yet implemented." }
+					]
+				}
+			]
+		},
+		{
+			"category": "InHouse",
+			"is_model_test": true,
+			"targets": [
+				{
+					"target": "NeuralMorphModel_global",
+					"tags": [ "SmokeTest" ],
+					"runtimes": [
+						{ "name": "NNXRuntimeORTDml", "relative_error": 0.02 },
+						{ "name": "NNXRuntimeHlsl",	"skip": true, "comment": "Model with multiple layer not yet implemented."	},
+						{ "name": "NNXRuntimeDml",  "skip": true, "comment": "Model with multiple layer not yet implemented." }
+					]
+				},
+				{
+					"target": "NeuralMorphModel_Local",
+					"tags": [ "SmokeTest" ],
+					"runtimes": [
+						{ "name": "NNXRuntimeORTDml", "relative_error": 0.04 },
+						{ "name": "NNXRuntimeHlsl",	"skip": true, "comment": "Model with multiple layer not yet implemented."	},
+						{ "name": "NNXRuntimeDml",  "skip": true, "comment": "Model with multiple layer not yet implemented." }
+					]
+				},
+				{
+					"target": "VertexDeltaModel",
+					"tags": [ "SmokeTest" ],
+					"runtimes": [
+						{ "name": "NNXRuntimeORTDml", "absolute_error": 1e-4, "relative_error": 10.0,
+							"comment": "TODO check, relative error is very high."
+						},
+						{ "name": "NNXRuntimeHlsl",	"skip": true, "comment": "Model with multiple layer not yet implemented."	},
+						{ "name": "NNXRuntimeDml",  "skip": true, "comment": "Model with multiple layer not yet implemented." }
+					]
+				}
+			]
+		}
+	],
+
+	"operator_test_categories": [
+		{
+			"category": "Elementwise.Unary",
+			"targets": [
+				{
+					"target": "Abs",
+					"additional_datasets": [ "Elementwise.Slow" ]
+				},
+				{
+					"target": "Tanh",
+					"runtimes": [
+						{ "name": "NNXRuntimeORTDml", "absolute_error":  1e-4, "relative_error": 0.5},
+						{ "name": "NNXRuntimeDml",    "absolute_error":  1e-4, "relative_error": 0.5},
+						{ "name": "NNXRuntimeHlsl",   "absolute_error":  1e-4, "relative_error": 0.5}
+					]
+				},
+				{
+					"target": "BitShift",
+					"skip": true,
+					"comment": "need attributes"
+				}
+			]
+		}
+	],
+
+	"input_output_sets": [
+		{
+			"name": "Elementwise.Unary.Regular",
+			"datasets": [
+				{ "inputs": [ { "shape": [ 2, 256, 256 ] } ], "comment": "if not specified output is defaulted as the first input" },
+				{ "inputs": [ { "shape": [ 1 ] } ] },
+				{ "inputs": [ { "shape": [ 4 ] } ] }
+			]
+		},
+		{
+			"name": "Elementwise.Unary.ZeroSized",
+			"datasets": [
+				{
+					"inputs": [ { "shape": [ 0 ] } ],
+					"runtimes": [
+						{ "name": "NNXRuntimeHLSL", "skip": true, "comment": "RDG does not support 0 sized tensor" },
+						{ "name": "NNXRuntimeDML",  "skip": true, "comment": "RDG does not support 0 sized tensor" }
+					]
+				}
+			]
+		},
+		{
+			"name": "Elementwise.Slow",
+			"datasets": [
+				{
+					"inputs": [ { "shape": [ 65536, 513 ] } ]
+				}
+			]
+		}
+	]
+	
+}
diff --git a/Engine/Plugins/Experimental/NNX/Source/NNXQA/Resources/NNXQATestDescSchema.json b/Engine/Plugins/Experimental/NNX/Source/NNXQA/Resources/NNXQATestDescSchema.json
new file mode 100644
--- /dev/null
+++ b/Engine/Plugins/Experimental/NNX/Source/NNXQA/Resources/NNXQATestDescSchema.json
@@ -0,0 +1,197 @@
+{
+	"description": "schema for NNX test",
+	"type": "object",
+	"additionalProperties": false,
+	"properties": {
+		"$schema": {
+			"type": "string"
+		},
+		"model_test_categories": {
+			"type": "array",
+			"items": { "$ref": "#/$defs/test_config_category" }
+		},
+		"operator_test_categories": {
+			"type": "array",
+			"items": { "$ref": "#/$defs/test_config_category" }
+		},
+		"input_output_sets": {
+			"type": "array",
+			"items": { "$ref": "#/$defs/test_config_input_output_set" }
+		},
+		"comment": {
+			"type": "string",
+			"description": "Unused by test system, usefull to add clarification in json file"
+		}
+	},
+
+	"$defs": {
+		"test_config_tensor": {
+			"type": "object",
+			"required": [ "shape" ],
+			"properties": {
+				"shape": {
+					"type": "array",
+					"items": { "type": "number" }
+				}
+			},
+			"comment": {
+				"type": "string",
+				"description": "Unused by test system, usefull to add clarification in json file"
+			}
+		},
+
+		"test_config_runtime": {
+			"type": "object",
+			"required": [ "name" ],
+			"additionalProperties": false,
+			"properties": {
+				"name": {
+					"type": "string",
+					"description": "name of the runtime to specialize for"
+				},
+				"skip": {
+					"type": "boolean",
+					"description": "true to skip the test for this runtime, false by default"
+				},
+				"absolute_error": {
+					"type": "number",
+					"minimum": 0,
+					"description": "max absolute error, 0 to keep engine default (1e-5)"
+				},
+				"relative_error": {
+					"type": "number",
+					"minimum": 0,
+					"maximum": 100,
+					"description": "max relative error in percent, 0 to keep engine default (1e-4 in %)"
+				},
+				"comment": {
+					"type": "string",
+					"description": "Unused by test system, usefull to add clarification in json file"
+				}
+			}
+		},
+
+		"test_config_dataset": {
+			"type": "object",
+			"required": [ "inputs" ],
+			"additionalProperties": false,
+			"properties": {
+				"inputs": {
+					"type": "array",
+					"description": "input shapes to test",
+					"items": { "$ref": "#/$defs/test_config_tensor" }
+				},
+				"outputs": {
+					"type": "array",
+					"description": "output shapes to test",
+					"items": { "$ref": "#/$defs/test_config_tensor" }
+				},
+				"runtimes": {
+					"type": "array",
+					"description": "Array of operator to test",
+					"items": { "$ref": "#/$defs/test_config_runtime" }
+				},
+				"comment": {
+					"type": "string",
+					"description": "Unused by test system, usefull to add clarification in json file"
+				}
+			}
+		},
+
+		"test_config_target": {
+			"type": "object",
+			"required": [ "target" ],
+			"additionalProperties": false,
+			"properties": {
+				"target": {
+					"type": "string",
+					"description": "Name of the operator or model to run (for model .onnx should be ommited)"
+				},
+				"tags": {
+					"type": "array",
+					"description": "array of tags associated with the test, test can be filtered by tag when runned from command line",
+					"items": { "type": "string" }
+				},
+				"additional_datasets": {
+					"type": "array",
+					"description": "additional dataset associated with this target, by default all dataset witch name are a super set of this category name will be runned, you can specify more dataset using this parameter",
+					"items": { "type": "string" }
+				},
+				"skip": {
+					"type": "boolean",
+					"description": "true to skip this operator or model test entirely, false by default"
+				},
+				"absolute_error": {
+					"type": "number",
+					"minimum": 0,
+					"description": "max absolute error, 0 to keep engine default (1e-5)"
+				},
+				"relative_error": {
+					"type": "number",
+					"minimum": 0,
+					"maximum": 100,
+					"description": "max relative error in percent, 0 to keep engine default (1e-4 in %)"
+				},
+				"runtimes": {
+					"type": "array",
+					"description": "Array of runtimes to specialize test behavior for",
+					"items": { "$ref": "#/$defs/test_config_runtime" }
+				},
+				"comment": {
+					"type": "string",
+					"description": "Unused by test system, usefull to add clarification in json file"
+				}
+			}
+		},
+
+		"test_config_category": {
+			"type": "object",
+			"required": [ "category", "targets" ],
+			"additionalProperties": false,
+			"properties": {
+				"category": {
+					"type": "string",
+					"description": "Name of the category"
+				},
+				"targets": {
+					"type": "array",
+					"description": "Array of models or operators to test",
+					"items": { "$ref": "#/$defs/test_config_target" }
+				},
+				"skip": {
+					"type": "boolean",
+					"description": "true to skip this category entirely, false by default"
+				},
+				"is_model_test": {
+					"type": "boolean",
+					"description": "true if the category contain model tests, false for operators tests"
+				},
+				"comment": {
+					"type": "string",
+					"description": "Unused by test system, usefull to add clarification in json file"
+				}
+			}
+		},
+
+		"test_config_input_output_set": {
+			"type": "object",
+			"required": [ "name", "datasets" ],
+			"additionalProperties": false,
+			"properties": {
+				"name": {
+					"type": "string",
+					"description": "Name of this dataset, any test name witch is a super set of this dataset name will use it"
+				},
+				"datasets": {
+					"type": "array",
+					"description": "Array of inputs/output to test",
+					"items": { "$ref": "#/$defs/test_config_dataset" }
+				},
+				"comment": {
+					"type": "string",
+					"description": "Unused by test system, usefull to add clarification in json file"
+				}
+			}
+		}
+	}
+}
\ No newline at end of file
