Fix conflicting NNXShaderParameters.h.

#rnx

diff --git a/Engine/Plugins/Experimental/NNX/Source/NNXHlslShaders/Public/NNXShaderParameters.h b/Engine/Plugins/Experimental/NNX/Source/NNXHlslShaders/Public/NNXShaderParameters.h
--- a/Engine/Plugins/Experimental/NNX/Source/NNXHlslShaders/Public/NNXShaderParameters.h
+++ b/Engine/Plugins/Experimental/NNX/Source/NNXHlslShaders/Public/NNXShaderParameters.h
@@ -7,17 +7,17 @@
 
 // HLSL runtime and DirectML runtime use different shader parameter access patterns
 #ifdef NNXRT_RDG_DML
-	#define NNXRT_RDG_BUFFER_SRV(View, BufferName) RDG_BUFFER_ACCESS(BufferName, ERHIAccess::UAVCompute)
-	#define NNXRT_RDG_BUFFER_UAV(View, BufferName) RDG_BUFFER_ACCESS(BufferName, ERHIAccess::UAVCompute)
+#define NNXRT_RDG_BUFFER_SRV(View, BufferName) RDG_BUFFER_ACCESS(BufferName, ERHIAccess::UAVCompute)
+#define NNXRT_RDG_BUFFER_UAV(View, BufferName) RDG_BUFFER_ACCESS(BufferName, ERHIAccess::UAVCompute)
 #else
-	#define NNXRT_RDG_BUFFER_SRV(View, BufferName) SHADER_PARAMETER_RDG_BUFFER_SRV(View, BufferName)
-	#define NNXRT_RDG_BUFFER_UAV(View, BufferName) SHADER_PARAMETER_RDG_BUFFER_UAV(View, BufferName)
+#define NNXRT_RDG_BUFFER_SRV(View, BufferName) SHADER_PARAMETER_RDG_BUFFER_SRV(View, BufferName)
+#define NNXRT_RDG_BUFFER_UAV(View, BufferName) SHADER_PARAMETER_RDG_BUFFER_UAV(View, BufferName)
 #endif
 
 #ifdef NNXRT_RDG_DML
-	#define NNXRT_BEGIN_SHADER_PARAMETER_STRUCT(Name) BEGIN_SHADER_PARAMETER_STRUCT(Name, )
+#define NNXRT_BEGIN_SHADER_PARAMETER_STRUCT(Name) BEGIN_SHADER_PARAMETER_STRUCT(Name, )
 #else
-	#define NNXRT_BEGIN_SHADER_PARAMETER_STRUCT(Name) BEGIN_SHADER_PARAMETER_STRUCT(FParameters, )
+#define NNXRT_BEGIN_SHADER_PARAMETER_STRUCT(Name) BEGIN_SHADER_PARAMETER_STRUCT(FParameters, )
 #endif
 
 #define NNXRT_END_SHADER_PARAMETER_STRUCT END_SHADER_PARAMETER_STRUCT
@@ -146,4 +146,40 @@ NNXRT_GEMM_PARAMETER_STRUCT()
 
 #ifdef NNXRT_RDG_DML
 NNXRT_CONV_PARAMETER_STRUCT()
+#endif
+
+//ConvTranspose
+#define NNXRT_CONVTRANSPOSE_MAX_NUM_STACK_DIMENSIONS 4
+#define NNXRT_CONVTRANSPOSE_MIN_NUM_READS_PER_THREAD_POW2 1
+#define NNXRT_CONVTRANSPOSE_MIN_NUM_READS_PER_THREAD_POW2 3
+
+#define NNXRT_CONVTRANSPOSE_PARAMETER_STRUCT() \
+	NNXRT_BEGIN_SHADER_PARAMETER_STRUCT(FMLConvParameters) \
+		NNXRT_RDG_BUFFER_SRV(Buffer<float>, X) \
+		NNXRT_RDG_BUFFER_SRV(Buffer<float>, W) \
+		NNXRT_RDG_BUFFER_UAV(RWBuffer<float>, Y) \
+		NNXRT_RDG_BUFFER_SRV(Buffer<float>, B) \
+		SHADER_PARAMETER_ARRAY(FIntVector4, Dilation_Stride_XBlockStartOffset_DilationXBlockStride, [NNXRT_CONVTRANSPOSE_MAX_NUM_STACK_DIMENSIONS]) \
+		SHADER_PARAMETER_ARRAY(FIntVector4, GroupStride_GroupShape_GroupThreadStride_StrideXBlockStride, [:NNXRT_CONVTRANSPOSE_MAX_NUM_STACK_DIMENSIONS])\
+		SHADER_PARAMETER_ARRAY(FIntVector4, YDimension_YMemoryStride_XDimension_XMemoryStride, NNXRT_CONVTRANSPOSE_MAX_NUM_STACK_DIMENSIONS])\
+		SHADER_PARAMETER_ARRAY(FIntVector4, XBlockStartStride_XBlockStride_WDimension_WDimensionDilationXBlockStride, [NNXRT_CONVTRANSPOSE_MAX_NUM_STACK_DIMENSIONS])\
+		SHADER_PARAMETER_ARRAY(FVector4f, OneDiv_GroupStride_GroupThreadStride_OneDivStride, [NNXRT_CONVTRANSPOSE_MAX_NUM_STACK_DIMENSIONS])\
+		SHADER_PARAMETER(int32, NumWChannels)\
+		SHADER_PARAMETER(int32, NumOutChannelsDivGroup)\
+		SHADER_PARAMETER(int32, YBatchStride)\
+		SHADER_PARAMETER(int32, YOutputKernelStride)\
+		SHADER_PARAMETER(int32, XBatchStride)\
+		SHADER_PARAMETER(int32, XChannelStride)\
+		SHADER_PARAMETER(int32, XBlockSize)\
+		SHADER_PARAMETER(int32, NumChannelBatches)\
+		SHADER_PARAMETER(int32, NumChannelsPerBatch)\
+		SHADER_PARAMETER(int32, WOutputKernelStride)\
+		SHADER_PARAMETER(int32, WChannelBatchSize)\
+		SHADER_PARAMETER(int32, WChannelSize)\
+		SHADER_PARAMETER(float, GroupsDivM)\
+		SHADER_PARAMETER(float, OneDivGroup)\
+	NNXRT_END_SHADER_PARAMETER_STRUCT()
+
+#ifdef NNXRT_RDG_DML
+NNXRT_CONVTRANSPOSE_PARAMETER_STRUCT()
 #endif
\ No newline at end of file
