Fix conflicting NNXShaderParameters.h.

#rnx

diff --git a/Engine/Plugins/Experimental/NNX/Source/NNXHlslShaders/Public/NNXConvTransposeCS.h b/Engine/Plugins/Experimental/NNX/Source/NNXHlslShaders/Public/NNXConvTransposeCS.h
--- a/Engine/Plugins/Experimental/NNX/Source/NNXHlslShaders/Public/NNXConvTransposeCS.h
+++ b/Engine/Plugins/Experimental/NNX/Source/NNXHlslShaders/Public/NNXConvTransposeCS.h
@@ -36,7 +36,7 @@ class NNXHLSLSHADERS_API FMLConvTransposeCS : public FGlobalShader
 	class FConvTransposeAlgorithm : SHADER_PERMUTATION_ENUM_CLASS("ALGORITHM", EConvTransposeAlgorithm);
 	class FConvTransposeGroupSize : SHADER_PERMUTATION_ENUM_CLASS("GROUP_SIZE", EConvTransposeGroupSize);
 	class FConvTransposeNumStackDimensions : SHADER_PERMUTATION_RANGE_INT("NUM_STACK_DIMENSIONS", 1, NNXRT_CONVTRANSPOSE_MAX_NUM_STACK_DIMENSIONS);
-	class FConvTransposeNumReadsPerThread : SHADER_PERMUTATION_RANGE_INT("NUM_READS_PER_THREAD_POW2", NNXRT_CONVTRANSPOSE_MIN_NUM_READS_PER_THREAD_POW2, NNXRT_CONVTRANSPOSE_MIN_NUM_READS_PER_THREAD_POW2);
+	class FConvTransposeNumReadsPerThread : SHADER_PERMUTATION_RANGE_INT("NUM_READS_PER_THREAD_POW2", NNXRT_CONVTRANSPOSE_MIN_NUM_READS_PER_THREAD_POW2, NNXRT_CONVTRANSPOSE_MAX_NUM_READS_PER_THREAD_POW2);
 	class FConvTransposeHasB : SHADER_PERMUTATION_BOOL("HAS_B");
 	using FPermutationDomain = TShaderPermutationDomain<FConvTransposeAlgorithm, FConvTransposeGroupSize, FConvTransposeNumStackDimensions, FConvTransposeNumReadsPerThread, FConvTransposeHasB>;
 
diff --git a/Engine/Plugins/Experimental/NNX/Source/NNXHlslShaders/Public/NNXShaderParameters.h b/Engine/Plugins/Experimental/NNX/Source/NNXHlslShaders/Public/NNXShaderParameters.h
--- a/Engine/Plugins/Experimental/NNX/Source/NNXHlslShaders/Public/NNXShaderParameters.h
+++ b/Engine/Plugins/Experimental/NNX/Source/NNXHlslShaders/Public/NNXShaderParameters.h
@@ -151,33 +151,33 @@ NNXRT_CONV_PARAMETER_STRUCT()
 //ConvTranspose
 #define NNXRT_CONVTRANSPOSE_MAX_NUM_STACK_DIMENSIONS 4
 #define NNXRT_CONVTRANSPOSE_MIN_NUM_READS_PER_THREAD_POW2 1
-#define NNXRT_CONVTRANSPOSE_MIN_NUM_READS_PER_THREAD_POW2 3
+#define NNXRT_CONVTRANSPOSE_MAX_NUM_READS_PER_THREAD_POW2 3
 
 #define NNXRT_CONVTRANSPOSE_PARAMETER_STRUCT() \
-	NNXRT_BEGIN_SHADER_PARAMETER_STRUCT(FMLConvParameters) \
+	NNXRT_BEGIN_SHADER_PARAMETER_STRUCT(FMLConvTransposeParameters) \
 		NNXRT_RDG_BUFFER_SRV(Buffer<float>, X) \
 		NNXRT_RDG_BUFFER_SRV(Buffer<float>, W) \
 		NNXRT_RDG_BUFFER_UAV(RWBuffer<float>, Y) \
 		NNXRT_RDG_BUFFER_SRV(Buffer<float>, B) \
 		SHADER_PARAMETER_ARRAY(FIntVector4, Dilation_Stride_XBlockStartOffset_DilationXBlockStride, [NNXRT_CONVTRANSPOSE_MAX_NUM_STACK_DIMENSIONS]) \
-		SHADER_PARAMETER_ARRAY(FIntVector4, GroupStride_GroupShape_GroupThreadStride_StrideXBlockStride, [:NNXRT_CONVTRANSPOSE_MAX_NUM_STACK_DIMENSIONS])\
-		SHADER_PARAMETER_ARRAY(FIntVector4, YDimension_YMemoryStride_XDimension_XMemoryStride, NNXRT_CONVTRANSPOSE_MAX_NUM_STACK_DIMENSIONS])\
-		SHADER_PARAMETER_ARRAY(FIntVector4, XBlockStartStride_XBlockStride_WDimension_WDimensionDilationXBlockStride, [NNXRT_CONVTRANSPOSE_MAX_NUM_STACK_DIMENSIONS])\
-		SHADER_PARAMETER_ARRAY(FVector4f, OneDiv_GroupStride_GroupThreadStride_OneDivStride, [NNXRT_CONVTRANSPOSE_MAX_NUM_STACK_DIMENSIONS])\
-		SHADER_PARAMETER(int32, NumWChannels)\
-		SHADER_PARAMETER(int32, NumOutChannelsDivGroup)\
-		SHADER_PARAMETER(int32, YBatchStride)\
-		SHADER_PARAMETER(int32, YOutputKernelStride)\
-		SHADER_PARAMETER(int32, XBatchStride)\
-		SHADER_PARAMETER(int32, XChannelStride)\
-		SHADER_PARAMETER(int32, XBlockSize)\
-		SHADER_PARAMETER(int32, NumChannelBatches)\
-		SHADER_PARAMETER(int32, NumChannelsPerBatch)\
-		SHADER_PARAMETER(int32, WOutputKernelStride)\
-		SHADER_PARAMETER(int32, WChannelBatchSize)\
-		SHADER_PARAMETER(int32, WChannelSize)\
-		SHADER_PARAMETER(float, GroupsDivM)\
-		SHADER_PARAMETER(float, OneDivGroup)\
+		SHADER_PARAMETER_ARRAY(FIntVector4, GroupStride_GroupShape_GroupThreadStride_StrideXBlockStride, [NNXRT_CONVTRANSPOSE_MAX_NUM_STACK_DIMENSIONS]) \
+		SHADER_PARAMETER_ARRAY(FIntVector4, YDimension_YMemoryStride_XDimension_XMemoryStride, [NNXRT_CONVTRANSPOSE_MAX_NUM_STACK_DIMENSIONS]) \
+		SHADER_PARAMETER_ARRAY(FIntVector4, XBlockStartStride_XBlockStride_WDimension_WDimensionDilationXBlockStride, [NNXRT_CONVTRANSPOSE_MAX_NUM_STACK_DIMENSIONS]) \
+		SHADER_PARAMETER_ARRAY(FVector4f, OneDiv_GroupStride_GroupThreadStride_OneDivStride, [NNXRT_CONVTRANSPOSE_MAX_NUM_STACK_DIMENSIONS]) \
+		SHADER_PARAMETER(int32, NumWChannels) \
+		SHADER_PARAMETER(int32, NumOutChannelsDivGroup) \
+		SHADER_PARAMETER(int32, YBatchStride) \
+		SHADER_PARAMETER(int32, YOutputKernelStride) \
+		SHADER_PARAMETER(int32, XBatchStride) \
+		SHADER_PARAMETER(int32, XChannelStride) \
+		SHADER_PARAMETER(int32, XBlockSize) \
+		SHADER_PARAMETER(int32, NumChannelBatches) \
+		SHADER_PARAMETER(int32, NumChannelsPerBatch) \
+		SHADER_PARAMETER(int32, WOutputKernelStride) \
+		SHADER_PARAMETER(int32, WChannelBatchSize) \
+		SHADER_PARAMETER(int32, WChannelSize) \
+		SHADER_PARAMETER(float, GroupsDivM) \
+		SHADER_PARAMETER(float, OneDivGroup) \
 	NNXRT_END_SHADER_PARAMETER_STRUCT()
 
 #ifdef NNXRT_RDG_DML
