NNX Plugin:
Fix Macro in OnnxruntimeCPU

diff --git a/Engine/Plugins/Experimental/NNX/Source/NNXRuntimeCPU/Public/NNXRuntimeCPU.h b/Engine/Plugins/Experimental/NNX/Source/NNXRuntimeCPU/Public/NNXRuntimeCPU.h
--- a/Engine/Plugins/Experimental/NNX/Source/NNXRuntimeCPU/Public/NNXRuntimeCPU.h
+++ b/Engine/Plugins/Experimental/NNX/Source/NNXRuntimeCPU/Public/NNXRuntimeCPU.h
@@ -15,7 +15,7 @@ NNI_THIRD_PARTY_INCLUDES_START
 #include "core/session/onnxruntime_cxx_api.h"
 NNI_THIRD_PARTY_INCLUDES_END
 
-#define NNX_RUNTIME_CPU_NAME L"NNXRuntimeCPU"
+#define NNX_RUNTIME_CPU_NAME TEXT("NNXRuntimeCPU")
 
 namespace NNX
 {
@@ -37,7 +37,6 @@ namespace NNX
 		virtual EMLRuntimeSupportFlags GetSupportFlags() const;
 		virtual FMLInferenceModel* CreateInferenceModel(UMLInferenceModel* Model);
 		FMLInferenceModel* CreateInferenceModel(UMLInferenceModel* InModel, const FMLInferenceNNXCPUConf& InConf);
-
 	};
 
 	static TUniquePtr<FRuntimeCPU> GCPURuntime;
