NNX Plugin:
Adding RHI for all platforms and RHI metal for Mac

diff --git a/Engine/Plugins/Experimental/NNX/Source/NNXRuntimeRDG/NNXRuntimeRDG.Build.cs b/Engine/Plugins/Experimental/NNX/Source/NNXRuntimeRDG/NNXRuntimeRDG.Build.cs
--- a/Engine/Plugins/Experimental/NNX/Source/NNXRuntimeRDG/NNXRuntimeRDG.Build.cs
+++ b/Engine/Plugins/Experimental/NNX/Source/NNXRuntimeRDG/NNXRuntimeRDG.Build.cs
@@ -23,21 +23,22 @@ public class NNXRuntimeRDG : ModuleRules
         PrivateDependencyModuleNames.AddRange(new string[]
         {
             "NNXCore",
-            "NNXHlslShaders"
+            "NNXHlslShaders",
+            "RHI"
         });
 
         if (Target.Platform == UnrealTargetPlatform.Win64)
 		{
-			PrivateDependencyModuleNames.AddRange(new string[]
-			{
-				"RHI",
-			});
-
 			PrivateDependencyModuleNames.Add("D3D12RHI");
 			PrivateDependencyModuleNames.Add("DirectMLDefault");
 
 			AddEngineThirdPartyPrivateStaticDependencies(Target, "DX12");
 			AddEngineThirdPartyPrivateStaticDependencies(Target, "DirectMLDefault");
 		}
+
+		if (Target.Platform == UnrealTargetPlatform.Mac)
+		{	
+			PrivateDependencyModuleNames.Add("MetalRHI");
+		}
 	}
 }
