NNX Plugin:
Deactivating DML and CUDA in all Platforms but Windows

diff --git a/Engine/Plugins/Experimental/NNX/Source/NNXRuntimeORT/Private/NNXRuntimeORTModule.cpp b/Engine/Plugins/Experimental/NNX/Source/NNXRuntimeORT/Private/NNXRuntimeORTModule.cpp
--- a/Engine/Plugins/Experimental/NNX/Source/NNXRuntimeORT/Private/NNXRuntimeORTModule.cpp
+++ b/Engine/Plugins/Experimental/NNX/Source/NNXRuntimeORT/Private/NNXRuntimeORTModule.cpp
@@ -60,20 +60,22 @@ void FNNXRuntimeORTModule::StartupModule()
 	FPlatformProcess::PopDllDirectory(*ORTDefaultRuntimeBinPath);
 
  	ORTRuntimeCPU = NNX::FRuntimeORTCPUStartup();
-	ORTRuntimeCUDA = NNX::FRuntimeORTCUDAStartup();
-	ORTRuntimeDML = NNX::FRuntimeORTDMLStartup();
 	if (ORTRuntimeCPU)
 	{
 		NNX::RegisterRuntime(ORTRuntimeCPU);
 	}
+#if PLATFORM_WINDOWS
+	ORTRuntimeCUDA = NNX::FRuntimeORTCUDAStartup();
 	if (ORTRuntimeCUDA)
 	{
 		NNX::RegisterRuntime(ORTRuntimeCUDA);
 	}
+	ORTRuntimeDML = NNX::FRuntimeORTDMLStartup();
 	if (ORTRuntimeDML)
 	{
 		NNX::RegisterRuntime(ORTRuntimeDML);
 	}
+#endif
 }
 
 // This function may be called during shutdown to clean up your module. For modules that support dynamic reloading,
@@ -85,19 +87,22 @@ void FNNXRuntimeORTModule::ShutdownModule()
 		NNX::UnregisterRuntime(ORTRuntimeCPU);
 		ORTRuntimeCPU = nullptr;
 	}
+	NNX::FRuntimeORTCPUShutdown();
+#if PLATFORM_WINDOWS
 	if (ORTRuntimeCUDA)
 	{
 		NNX::UnregisterRuntime(ORTRuntimeCUDA);
 		ORTRuntimeCUDA = nullptr;
 	}
+	NNX::FRuntimeORTCUDAShutdown();
 	if (ORTRuntimeDML)
 	{
 		NNX::UnregisterRuntime(ORTRuntimeDML);
 		ORTRuntimeDML = nullptr;
 	}
-	NNX::FRuntimeORTCPUShutdown();
 	NNX::FRuntimeORTDMLShutdown();
-	NNX::FRuntimeORTCUDAShutdown();
+#endif
+
 }
 
 IMPLEMENT_MODULE(FNNXRuntimeORTModule, NNXRuntimeORT);
