NNXQA Add all unary operator to parametric test but the one requiring bool tensors

diff --git a/Engine/Plugins/Experimental/NNX/Source/NNXQA/Resources/NNXQATestDesc.json b/Engine/Plugins/Experimental/NNX/Source/NNXQA/Resources/NNXQATestDesc.json
--- a/Engine/Plugins/Experimental/NNX/Source/NNXQA/Resources/NNXQATestDesc.json
+++ b/Engine/Plugins/Experimental/NNX/Source/NNXQA/Resources/NNXQATestDesc.json
@@ -69,23 +69,88 @@
 		{
 			"category": "Elementwise.Unary",
 			"targets": [
-				{
-					"target": "Abs",
-					"additional_datasets": [ "Elementwise.Slow" ]
-				},
-				{
-					"target": "Tanh",
-					"runtimes": [
-						{ "name": "NNXRuntimeORTDml", "absolute_error":  1e-4, "relative_error": 0.5},
-						{ "name": "NNXRuntimeDml",    "absolute_error":  1e-4, "relative_error": 0.5},
-						{ "name": "NNXRuntimeHlsl",   "absolute_error":  1e-4, "relative_error": 0.5}
-					]
-				},
-				{
-					"target": "BitShift",
-					"skip": true,
-					"comment": "need attributes"
-				}
+				{	"target": "Abs", "additional_datasets": [ "Elementwise.Slow" ]},
+				{	"target": "Acos", "runtimes": [
+						{ "name": "NNXRuntimeORTDml", "absolute_error":  1e-4, "relative_error": 0.03},
+						{ "name": "NNXRuntimeDml",    "absolute_error":  1e-4, "relative_error": 0.03},
+						{ "name": "NNXRuntimeHlsl",   "absolute_error":  1e-4, "relative_error": 0.03}
+					]},
+				{	"target": "Acosh" },
+				{	"target": "Asin", "runtimes": [
+						{ "name": "NNXRuntimeORTDml", "absolute_error":  1e-4, "relative_error": 0.03},
+						{ "name": "NNXRuntimeDml",    "absolute_error":  1e-4, "relative_error": 1.5},
+						{ "name": "NNXRuntimeHlsl",   "absolute_error":  1e-4, "relative_error": 1.5}
+					]},
+				{	"target": "Asinh", "runtimes": [
+						{ "name": "NNXRuntimeORTDml", "absolute_error":  1e-4, "relative_error": 0.03},
+						{ "name": "NNXRuntimeDml",    "absolute_error":  1e-4, "relative_error": 0.03},
+						{ "name": "NNXRuntimeHlsl",   "absolute_error":  1e-4, "relative_error": 0.03}
+					]},
+				{	"target": "Atan", "runtimes": [
+						{ "name": "NNXRuntimeORTDml", "absolute_error":  1e-4, "relative_error": 0.03},
+						{ "name": "NNXRuntimeDml",    "absolute_error":  1e-4, "relative_error": 0.03},
+						{ "name": "NNXRuntimeHlsl",   "absolute_error":  1e-4, "relative_error": 0.03}
+					]},
+				{	"target": "Atanh" },
+				{	"target": "BitShift", "skip": true, "comment": "//TODO need attributes" },
+				{	"target": "Cast",     "skip": true, "comment": "//TODO need attributes" },
+				{	"target": "Ceil" },
+				{	"target": "Clip", "additional_datasets": [ "Elementwise.ClipTestWithExtraInput" ], "skip": true, "comment": "//Support bool tensor in parametric tests" },
+				{	"target": "Cos" },
+				{	"target": "Cosh" },
+				{	"target": "Elu", "comment": "//TODO test with not-default attribute", "runtimes": [
+						{ "name": "NNXRuntimeORTDml", "absolute_error":  1e-4, "relative_error": 0.03},
+						{ "name": "NNXRuntimeDml",    "absolute_error":  1e-4, "relative_error": 0.03},
+						{ "name": "NNXRuntimeHlsl",   "absolute_error":  1e-4, "relative_error": 0.03}
+					]},
+				{	"target": "Erf", "runtimes": [
+						{ "name": "NNXRuntimeORTDml", "absolute_error":  1e-4, "relative_error": 0.03},
+						{ "name": "NNXRuntimeDml",    "absolute_error":  1e-4, "relative_error": 0.03},
+						{ "name": "NNXRuntimeHlsl",   "absolute_error":  5e-4, "relative_error": 0.05, "comment": "NNXRuntimeHLSL Erf implementation is less precise than the DML one"}
+					]},
+				{	"target": "Exp" },
+				{	"target": "Floor" },
+				{	"target": "IsInf", "skip": true, "comment": "//TODO add flexibility in test data setup + bool tensors" },
+				{	"target": "IsNan", "skip": true, "comment": "//TODO add flexibility in test data setup + bool tensors" },
+				{	"target": "HardSigmoid",         "comment": "//TODO test with not-default attribute" },
+				{	"target": "HardSwish" },
+				{	"target": "LeakyRelu",           "comment": "//TODO test with not-default attribute" },
+				{	"target": "Log", "runtimes": [
+						{ "name": "NNXRuntimeORTDml", "relative_error": 0.15},
+						{ "name": "NNXRuntimeDml",    "relative_error": 0.15},
+						{ "name": "NNXRuntimeHlsl",   "relative_error": 0.15}
+					]},
+				{	"target": "Neg" },
+				{	"target": "Not", "skip": true, "comment": "//Support bool tensor in parametric tests" },
+				{	"target": "Reciprocal" },
+				{	"target": "Relu" },
+				{	"target": "Round" },
+				{	"target": "Selu", "comment": "//TODO test with not-default attribute", "runtimes": [
+						{ "name": "NNXRuntimeORTDml", "absolute_error":  1e-4, "relative_error": 0.03},
+						{ "name": "NNXRuntimeDml",    "absolute_error":  1e-4, "relative_error": 0.03},
+						{ "name": "NNXRuntimeHlsl",   "absolute_error":  1e-4, "relative_error": 0.03}
+					]},
+				{	"target": "Sigmoid" },
+				{	"target": "Sign" },
+				{	"target": "Sin", "runtimes": [
+						{ "name": "NNXRuntimeORTDml", "absolute_error":  1e-4, "relative_error": 0.03},
+						{ "name": "NNXRuntimeDml",    "absolute_error":  1e-4, "relative_error": 0.03},
+						{ "name": "NNXRuntimeHlsl",   "absolute_error":  1e-4, "relative_error": 0.03}
+					]},
+				{	"target": "Sinh" },
+				{	"target": "Softplus" },
+				{	"target": "Softsign" },
+				{	"target": "Sqrt" },
+				{	"target": "Tan", "runtimes": [
+						{ "name": "NNXRuntimeORTDml", "absolute_error":  1e-4, "relative_error": 0.03},
+						{ "name": "NNXRuntimeDml",    "absolute_error":  1e-4, "relative_error": 0.03},
+						{ "name": "NNXRuntimeHlsl",   "absolute_error":  1e-4, "relative_error": 0.03}
+					]},
+				{	"target": "Tanh", "runtimes": [
+						{ "name": "NNXRuntimeORTDml", "absolute_error":  1e-4, "relative_error": 0.03},
+						{ "name": "NNXRuntimeDml",    "absolute_error":  1e-4, "relative_error": 0.03},
+						{ "name": "NNXRuntimeHlsl",   "absolute_error":  1e-4, "relative_error": 0.03}
+					]}
 			]
 		}
 	],
@@ -94,21 +159,26 @@
 		{
 			"name": "Elementwise.Unary.Regular",
 			"datasets": [
-				{ "inputs": [ { "shape": [ 2, 256, 256 ] } ], "comment": "if not specified output is defaulted as the first input" },
-				{ "inputs": [ { "shape": [ 1 ] } ] },
-				{ "inputs": [ { "shape": [ 4 ] } ] }
+				{ "inputs": [ { "shape": [ 1 ] } ], "comment": "if not specified output is defaulted as the first input" },
+				{ "inputs": [ { "shape": [ 1,2,3,4 ] } ] },
+				{ "inputs": [ { "shape": [ 1,512 ] } ] }
 			]
 		},
 		{
 			"name": "Elementwise.Unary.ZeroSized",
 			"datasets": [
-				{
-					"inputs": [ { "shape": [ 0 ] } ],
+				{ "inputs": [ { "shape": [ 1,0,3 ] } ],
 					"runtimes": [
 						{ "name": "NNXRuntimeHLSL", "skip": true, "comment": "RDG does not support 0 sized tensor" },
 						{ "name": "NNXRuntimeDML",  "skip": true, "comment": "RDG does not support 0 sized tensor" }
-					]
-				}
+					]}
+			]
+		},
+		{
+			"name": "Elementwise.ClipTestWithExtraInput",
+			"datasets": [
+				{ "inputs": [ { "shape": [1], "comment": "min" } ] },
+				{ "inputs": [ { "shape": [1], "comment": "min" }, { "shape": [1], "comment": "max" } ] }
 			]
 		},
 		{
@@ -118,7 +188,7 @@
 					"inputs": [ { "shape": [ 65536, 513 ] } ]
 				}
 			]
-		}
+		},
 	]
 	
 }
