Add VulkanRHI to NNXRuntimeRDG

diff --git a/Engine/Plugins/Experimental/NNX/Source/NNXRuntimeRDG/NNXRuntimeRDG.Build.cs b/Engine/Plugins/Experimental/NNX/Source/NNXRuntimeRDG/NNXRuntimeRDG.Build.cs
--- a/Engine/Plugins/Experimental/NNX/Source/NNXRuntimeRDG/NNXRuntimeRDG.Build.cs
+++ b/Engine/Plugins/Experimental/NNX/Source/NNXRuntimeRDG/NNXRuntimeRDG.Build.cs
@@ -40,5 +40,10 @@ public class NNXRuntimeRDG : ModuleRules
 		{	
 			PrivateDependencyModuleNames.Add("MetalRHI");
 		}
+
+		if (Target.Platform == UnrealTargetPlatform.Linux)
+		{	
+			PrivateDependencyModuleNames.Add("VulkanRHI");
+		}
 	}
 }
