NNX: Better log for ORT session option creation + deactivate ORTCuda test for now

diff --git a/Engine/Plugins/Experimental/NNX/Source/NNXQA/Private/NNXQAUtils.cpp b/Engine/Plugins/Experimental/NNX/Source/NNXQA/Private/NNXQAUtils.cpp
--- a/Engine/Plugins/Experimental/NNX/Source/NNXQA/Private/NNXQAUtils.cpp
+++ b/Engine/Plugins/Experimental/NNX/Source/NNXQA/Private/NNXQAUtils.cpp
@@ -417,6 +417,12 @@ namespace Test
 				continue;
 			}
 
+			if (RuntimeName == "NNXRuntimeORTCuda")
+			{
+				//TODO Reactivate tests for NNXRuntimeORTCuda runtime. Skipped for now as we wait for legal approval for the dlls.
+				continue;
+			}
+
 			FString TestResult;
 
 			if (AutomationTestSetup != nullptr && (
diff --git a/Engine/Plugins/Experimental/NNX/Source/NNXRuntimeORT/Private/NNXRuntimeORT.cpp b/Engine/Plugins/Experimental/NNX/Source/NNXRuntimeORT/Private/NNXRuntimeORT.cpp
--- a/Engine/Plugins/Experimental/NNX/Source/NNXRuntimeORT/Private/NNXRuntimeORT.cpp
+++ b/Engine/Plugins/Experimental/NNX/Source/NNXRuntimeORT/Private/NNXRuntimeORT.cpp
@@ -370,7 +370,7 @@ bool FMLInferenceModelORTCuda::InitializedAndConfigureMembers()
 	OrtStatusPtr Status = OrtSessionOptionsAppendExecutionProvider_CUDA(*SessionOptions.Get(), ORTConfiguration.DeviceId);
 	if (Status)
 	{
-		UE_LOG(LogNNX, Warning, TEXT("Failed to initialize session options for ORT CUDA EP"));
+		UE_LOG(LogNNX, Warning, TEXT("Failed to initialize session options for ORT CUDA EP: %s"), ANSI_TO_TCHAR(Ort::GetApi().GetErrorMessage(Status)));
 		return false;
 	}
 
@@ -449,7 +449,7 @@ bool FMLInferenceModelORTDml::InitializedAndConfigureMembers()
 	OrtStatusPtr Status = OrtSessionOptionsAppendExecutionProviderEx_DML(*SessionOptions.Get(), DmlDevice, CmdQ);
 	if (Status)
 	{
-		UE_LOG(LogNNX, Warning, TEXT("Failed to initialize session options for ORT Dml EP"));
+		UE_LOG(LogNNX, Warning, TEXT("Failed to initialize session options for ORT Dml EP: %s"), ANSI_TO_TCHAR(Ort::GetApi().GetErrorMessage(Status)));
 		return false;
 	}
 
