NNXQA tensor type support for parametric tests

diff --git a/Engine/Plugins/Experimental/NNX/Source/NNXQA/Private/NNXQAJsonUtils.h b/Engine/Plugins/Experimental/NNX/Source/NNXQA/Private/NNXQAJsonUtils.h
--- a/Engine/Plugins/Experimental/NNX/Source/NNXQA/Private/NNXQAJsonUtils.h
+++ b/Engine/Plugins/Experimental/NNX/Source/NNXQA/Private/NNXQAJsonUtils.h
@@ -17,6 +17,7 @@ namespace Json
 	struct FTestConfigTensor : FJsonSerializable
 	{
 		TArray<int32> Shape;
+		FString Type;
 		//TODO Extend:
 		// by adding custom data
 		// by adding custom initializer or parameter (range, distribution)
@@ -25,7 +26,7 @@ namespace Json
 
 		BEGIN_JSON_SERIALIZER
 			JSON_SERIALIZE_ARRAY("shape", Shape);
-			//JSON_SERIALIZE("type", Type);
+			JSON_SERIALIZE("data_type", Type);
 			//JSON_SERIALIZE("on_gpu", OnGpu);
 			//JSON_SERIALIZE("range", ...);
 			//JSON_SERIALIZE("data_source", DataSource);
@@ -68,6 +69,8 @@ namespace Json
 		bool Skip;
 		float AbsoluteError;
 		float RelativeError;
+		FString InputType;
+		FString OutputType;
 		TArray<FTestConfigRuntime> Runtimes;
 
 		BEGIN_JSON_SERIALIZER
@@ -75,6 +78,8 @@ namespace Json
 			JSON_SERIALIZE("skip", Skip);
 			JSON_SERIALIZE("absolute_error", AbsoluteError);
 			JSON_SERIALIZE("relative_error", RelativeError);
+			JSON_SERIALIZE("input_type", InputType);
+			JSON_SERIALIZE("output_type", OutputType);
 			JSON_SERIALIZE_ARRAY_SERIALIZABLE("runtimes", Runtimes, FTestConfigRuntime);
 			JSON_SERIALIZE_ARRAY("tags", Tags);
 			JSON_SERIALIZE_ARRAY("additional_datasets", AdditionalDatasets);
diff --git a/Engine/Plugins/Experimental/NNX/Source/NNXQA/Private/NNXQAParametricTest.cpp b/Engine/Plugins/Experimental/NNX/Source/NNXQA/Private/NNXQAParametricTest.cpp
--- a/Engine/Plugins/Experimental/NNX/Source/NNXQA/Private/NNXQAParametricTest.cpp
+++ b/Engine/Plugins/Experimental/NNX/Source/NNXQA/Private/NNXQAParametricTest.cpp
@@ -1,13 +1,17 @@
 // Copyright Epic Games, Inc. All Rights Reserved.
 
 #include "NNXQAParametricTest.h"
+
+#include "UObject/Class.h"
 #include "NNXCore.h"
+#include "NNXTypes.h"
 #include "NNXQAUtils.h"
 #include "NNXQAJsonUtils.h"
 #include "NNXModelBuilder.h"
 #include "Interfaces/IPluginManager.h"
 #include "Misc/Paths.h"
 #include "HAL/IConsoleManager.h"
+#include "UObject/ReflectedTypeAccessors.h"
 
 #if WITH_DEV_AUTOMATION_TESTS
 #include "Misc/AutomationTest.h"
@@ -123,7 +127,14 @@ namespace Test
 			return Shape;
 		}
 
-		static void ApplyDatasetConfig(FTests::FTestSetup& TestSetup, const Json::FTestConfigDataset& TestDataset)
+		static EMLTensorDataType GetTensorTypeFromJson(const FString& TypeName, EMLTensorDataType DefaultValue)
+		{
+			int64 Value = StaticEnum<EMLTensorDataType>()->GetValueByNameString(TypeName);
+			
+			return (Value == INDEX_NONE) ? DefaultValue : (EMLTensorDataType)Value;
+		}
+
+		static void ApplyDatasetConfig(FTests::FTestSetup& TestSetup, const Json::FTestConfigDataset& TestDataset, EMLTensorDataType DefaultInputType, EMLTensorDataType DefaultOutputType)
 		{
 			ApplyRuntimesConfig(TestSetup, TestDataset.Runtimes);
 
@@ -132,13 +143,11 @@ namespace Test
 				return;
 			}
 
-			//TODO allow more tensor type
-			EMLTensorDataType TensorType = EMLTensorDataType::Float;
-
 			uint32 i = 0;
 			for (auto&& Tensor : TestDataset.Inputs)
 			{
 				TArray<uint32> shape = GetShapeFromJsonArray(Tensor.Shape);
+				EMLTensorDataType TensorType = GetTensorTypeFromJson(Tensor.Type, DefaultInputType);
 				TestSetup.Inputs.Emplace(FMLTensorDesc::Make(FString::Printf(TEXT("in%d"), i++), shape, TensorType));
 			}
 
@@ -146,12 +155,14 @@ namespace Test
 			for (auto&& Tensor : TestDataset.Outputs)
 			{
 				TArray<uint32> shape = GetShapeFromJsonArray(Tensor.Shape);
+				EMLTensorDataType TensorType = GetTensorTypeFromJson(Tensor.Type, DefaultOutputType);
 				TestSetup.Outputs.Emplace(FMLTensorDesc::Make(FString::Printf(TEXT("output%d"), i++), shape, TensorType));
 			}
 			//If output is not defined it is the first input shape.
 			if (TestDataset.Outputs.Num() == 0 && TestDataset.Inputs.Num() > 0)
 			{
 				TArray<uint32> shape = GetShapeFromJsonArray(TestDataset.Inputs[0].Shape);
+				EMLTensorDataType TensorType = GetTensorTypeFromJson(TestDataset.Inputs[0].Type, DefaultOutputType);
 				TestSetup.Outputs.Emplace(FMLTensorDesc::Make(TEXT("output"), shape, TensorType));
 			}
 		}
@@ -195,6 +206,7 @@ namespace Test
 
 				const bool bIsModelCategory = TestCategory.IsModelTest;
 				const FString TestCategoryPath(BaseTestPath + TestCategory.Category + TEXT("."));
+				
 				for (const Json::FTestConfigTarget& TestTarget : TestCategory.Targets)
 				{
 					if (TestTarget.Skip)
@@ -203,8 +215,10 @@ namespace Test
 					}
 
 					const FString& TestBaseName = TestTarget.Target;
-
+					EMLTensorDataType InputTypeFromTarget = GetTensorTypeFromJson(TestTarget.InputType, EMLTensorDataType::Float);
+					EMLTensorDataType OutputTypeFromTarget = GetTensorTypeFromJson(TestTarget.InputType, EMLTensorDataType::Float);
 					bool bAtLeastATestWasAdded = false;
+					
 					for (const Json::FTestConfigInputOutputSet& InputOutputSet : InputOutputSets)
 					{
 						//If category is a substring of InputOutputSet name or if target explicitly requested the InputOutputSet name
@@ -218,18 +232,20 @@ namespace Test
 								}
 
 								FTests::FTestSetup& Test = AddTest(TestCategoryPath, TestBaseName, TEXT(".") + GetTestSuffix(Dataset));
+								
 								ApplyTargetConfig(Test, TestTarget);
-								ApplyDatasetConfig(Test, Dataset);
+								ApplyDatasetConfig(Test, Dataset, InputTypeFromTarget, OutputTypeFromTarget);
 								Test.IsModelTest = bIsModelCategory;
 							}
 							bAtLeastATestWasAdded = true;
 						}
 					}
 
-					//No dataset were matched with this target, define a test without enforcing shapes.
+					//No dataset were matched with this target, define a test without enforcing input/output.
 					if (!bAtLeastATestWasAdded)
 					{
 						FTests::FTestSetup& Test = AddTest(TestCategoryPath, TestBaseName, TEXT(""));
+						
 						ApplyTargetConfig(Test, TestTarget);
 						Test.IsModelTest = bIsModelCategory;
 					}
@@ -254,6 +270,7 @@ namespace Test
 			// Model test, load model from disk
 			FString ModelPath = GetFullModelPath(TestSetup.TargetName + TEXT(".onnx"));
 			const bool bIsModelInMem = FFileHelper::LoadFileToArray(ModelData, *ModelPath);
+			
 			if (!bIsModelInMem)
 			{
 				UE_LOG(LogNNX, Error, TEXT("Can't load model from disk at path '%s'. Tests ABORTED!"), *ModelPath);
@@ -301,6 +318,7 @@ namespace Test
 		bool bAllTestSucceed = true;
 		uint32 NumTest = 0;
 		uint32 NumTestFailed = 0;
+		
 		for (FTests::FTestSetup& Test : ParametricTests.TestSetups)
 		{
 			if (!Tag.IsEmpty() && !Test.Tags.Contains(Tag))
@@ -340,6 +358,7 @@ namespace Test
 	{
 		FString Arg(TEXT(""));
 		int32 ArgNameIndex = Args.Find(ArgName);
+		
 		if (ArgNameIndex == INDEX_NONE || ArgNameIndex + 1 >= Args.Num())
 		{
 			return Arg;
diff --git a/Engine/Plugins/Experimental/NNX/Source/NNXQA/Resources/NNXQATestDesc.json b/Engine/Plugins/Experimental/NNX/Source/NNXQA/Resources/NNXQATestDesc.json
--- a/Engine/Plugins/Experimental/NNX/Source/NNXQA/Resources/NNXQATestDesc.json
+++ b/Engine/Plugins/Experimental/NNX/Source/NNXQA/Resources/NNXQATestDesc.json
@@ -95,7 +95,10 @@
 				{	"target": "BitShift", "skip": true, "comment": "//TODO need attributes" },
 				{	"target": "Cast",     "skip": true, "comment": "//TODO need attributes" },
 				{	"target": "Ceil" },
-				{	"target": "Clip", "additional_datasets": [ "Elementwise.ClipTestWithExtraInput" ], "skip": true, "comment": "//Support bool tensor in parametric tests" },
+				{	"target": "Clip", "additional_datasets": [ "Elementwise.ClipTestWithExtraInput" ], "runtimes": [
+						{ "name": "NNXRuntimeDml", "skip": true, "comment": "RDG need scalar tensor support"},
+						{ "name": "NNXRuntimeHlsl", "skip": true, "comment": "RDG need scalar tensor support"}						
+					]},
 				{	"target": "Cos" },
 				{	"target": "Cosh" },
 				{	"target": "Elu", "comment": "//TODO test with not-default attribute", "runtimes": [
@@ -121,7 +124,10 @@
 						{ "name": "NNXRuntimeHlsl",   "relative_error": 0.15}
 					]},
 				{	"target": "Neg" },
-				{	"target": "Not", "skip": true, "comment": "//Support bool tensor in parametric tests" },
+				{	"target": "Not", "input_type": "boolean", "output_type": "boolean", "runtimes": [
+						{ "name": "NNXRuntimeDml", "skip": true, "comment": "RDG need bool tensor support"},
+						{ "name": "NNXRuntimeHlsl", "skip": true, "comment": "RDG need bool tensor support"}						
+					]},
 				{	"target": "Reciprocal" },
 				{	"target": "Relu" },
 				{	"target": "Round" },
@@ -177,8 +183,8 @@
 		{
 			"name": "Elementwise.ClipTestWithExtraInput",
 			"datasets": [
-				{ "inputs": [ { "shape": [1], "comment": "min" } ] },
-				{ "inputs": [ { "shape": [1], "comment": "min" }, { "shape": [1], "comment": "max" } ] }
+				{ "inputs": [ { "shape": [1], "comment": "input" }, { "shape": [1], "comment": "min" } ] },
+				{ "inputs": [ { "shape": [1], "comment": "input" }, { "shape": [1], "comment": "min" }, { "shape": [1], "comment": "max" } ] }
 			]
 		},
 		{
diff --git a/Engine/Plugins/Experimental/NNX/Source/NNXQA/Resources/NNXQATestDescSchema.json b/Engine/Plugins/Experimental/NNX/Source/NNXQA/Resources/NNXQATestDescSchema.json
--- a/Engine/Plugins/Experimental/NNX/Source/NNXQA/Resources/NNXQATestDescSchema.json
+++ b/Engine/Plugins/Experimental/NNX/Source/NNXQA/Resources/NNXQATestDescSchema.json
@@ -32,6 +32,10 @@
 				"shape": {
 					"type": "array",
 					"items": { "type": "number" }
+			    },
+				"data_type": {
+					"type": "string",
+					"enum": [ "float", "int32", "uint32", "boolean" ]
 				},
 				"comment": {
 					"type": "string",
@@ -132,6 +136,14 @@
 					"maximum": 100,
 					"description": "max relative error in percent, 0 to keep engine default (1e-4 in %)"
 				},
+				"input_type": {
+					"type": "string",
+					"enum": [ "float", "int32", "uint32", "boolean" ]
+				},
+				"output_type": {
+					"type": "string",
+					"enum": [ "float", "int32", "uint32", "boolean" ]
+				},
 				"runtimes": {
 					"type": "array",
 					"description": "Array of runtimes to specialize test behavior for",
