NNXQA Add elem wise binary op to parametric tests

diff --git a/Engine/Plugins/Experimental/NNX/Source/NNXQA/Private/NNXQAJsonUtils.h b/Engine/Plugins/Experimental/NNX/Source/NNXQA/Private/NNXQAJsonUtils.h
--- a/Engine/Plugins/Experimental/NNX/Source/NNXQA/Private/NNXQAJsonUtils.h
+++ b/Engine/Plugins/Experimental/NNX/Source/NNXQA/Private/NNXQAJsonUtils.h
@@ -66,6 +66,7 @@ namespace Json
 		FString Target;
 		TArray<FString> Tags;
 		TArray<FString> AdditionalDatasets;
+		TArray<FString> RemovedDatasets;
 		bool Skip;
 		float AbsoluteError;
 		float RelativeError;
@@ -83,6 +84,7 @@ namespace Json
 			JSON_SERIALIZE_ARRAY_SERIALIZABLE("runtimes", Runtimes, FTestConfigRuntime);
 			JSON_SERIALIZE_ARRAY("tags", Tags);
 			JSON_SERIALIZE_ARRAY("additional_datasets", AdditionalDatasets);
+			JSON_SERIALIZE_ARRAY("removed_datasets", RemovedDatasets);
 		END_JSON_SERIALIZER
 	};
 
@@ -90,6 +92,7 @@ namespace Json
 	{
 		FString Category;
 		TArray<FTestConfigTarget> Targets;
+		TArray<FTestConfigRuntime> Runtimes;
 		bool Skip;
 		bool IsModelTest;
 
@@ -98,6 +101,7 @@ namespace Json
 			JSON_SERIALIZE("skip", Skip);
 			JSON_SERIALIZE("is_model_test", IsModelTest);
 			JSON_SERIALIZE_ARRAY_SERIALIZABLE("targets", Targets, FTestConfigTarget);
+			JSON_SERIALIZE_ARRAY_SERIALIZABLE("runtimes", Runtimes, FTestConfigRuntime);
 		END_JSON_SERIALIZER
 	};
 
diff --git a/Engine/Plugins/Experimental/NNX/Source/NNXQA/Private/NNXQAParametricTest.cpp b/Engine/Plugins/Experimental/NNX/Source/NNXQA/Private/NNXQAParametricTest.cpp
--- a/Engine/Plugins/Experimental/NNX/Source/NNXQA/Private/NNXQAParametricTest.cpp
+++ b/Engine/Plugins/Experimental/NNX/Source/NNXQA/Private/NNXQAParametricTest.cpp
@@ -216,7 +216,7 @@ namespace Test
 
 					const FString& TestBaseName = TestTarget.Target;
 					EMLTensorDataType InputTypeFromTarget = GetTensorTypeFromJson(TestTarget.InputType, EMLTensorDataType::Float);
-					EMLTensorDataType OutputTypeFromTarget = GetTensorTypeFromJson(TestTarget.InputType, EMLTensorDataType::Float);
+					EMLTensorDataType OutputTypeFromTarget = GetTensorTypeFromJson(TestTarget.OutputType, EMLTensorDataType::Float);
 					bool bAtLeastATestWasAdded = false;
 					
 					for (const Json::FTestConfigInputOutputSet& InputOutputSet : InputOutputSets)
@@ -224,6 +224,11 @@ namespace Test
 						//If category is a substring of InputOutputSet name or if target explicitly requested the InputOutputSet name
 						if (InputOutputSet.Name.Contains(TestCategory.Category) || TestTarget.AdditionalDatasets.Contains(InputOutputSet.Name))
 						{
+							if (TestTarget.RemovedDatasets.Contains(InputOutputSet.Name))
+							{
+								continue;
+							}
+							
 							for (const Json::FTestConfigDataset& Dataset : InputOutputSet.Datasets)
 							{
 								if (Dataset.Inputs.Num() == 0)
@@ -233,6 +238,7 @@ namespace Test
 
 								FTests::FTestSetup& Test = AddTest(TestCategoryPath, TestBaseName, TEXT(".") + GetTestSuffix(Dataset));
 								
+								ApplyRuntimesConfig(Test, TestCategory.Runtimes);
 								ApplyTargetConfig(Test, TestTarget);
 								ApplyDatasetConfig(Test, Dataset, InputTypeFromTarget, OutputTypeFromTarget);
 								Test.IsModelTest = bIsModelCategory;
@@ -246,6 +252,7 @@ namespace Test
 					{
 						FTests::FTestSetup& Test = AddTest(TestCategoryPath, TestBaseName, TEXT(""));
 						
+						ApplyRuntimesConfig(Test, TestCategory.Runtimes);
 						ApplyTargetConfig(Test, TestTarget);
 						Test.IsModelTest = bIsModelCategory;
 					}
diff --git a/Engine/Plugins/Experimental/NNX/Source/NNXQA/Resources/NNXQATestDesc.json b/Engine/Plugins/Experimental/NNX/Source/NNXQA/Resources/NNXQATestDesc.json
--- a/Engine/Plugins/Experimental/NNX/Source/NNXQA/Resources/NNXQATestDesc.json
+++ b/Engine/Plugins/Experimental/NNX/Source/NNXQA/Resources/NNXQATestDesc.json
@@ -19,7 +19,7 @@
 					"tags": [ "SmokeTest" ],
 					"runtimes": [
 						{ "name": "NNXRuntimeORTDml", "absolute_error": 5e-2, "relative_error": 50.0,
-							"comment": "//Worst relative error 43.367836 % (epsilon 0.001000 %) at position 112558, got -0.003391 expected -0.002365. TODO we need an automatic balance between relative error and absolute error"
+							"comment": "Worst relative error 43.367836 % (epsilon 0.001000 %) at position 112558, got -0.003391 expected -0.002365. TODO we need an automatic balance between relative error and absolute error"
 						},
 						{ "name": "NNXRuntimeCPU", "absolute_error": 1e-3, "relative_error": 3.0 },
 						{ "name": "NNXRuntimeHlsl",	"skip": true, "comment": "Model with multiple layer not yet implemented."	},
@@ -68,8 +68,9 @@
 	"operator_test_categories": [
 		{
 			"category": "Elementwise.Unary",
+			"runtimes": [ { "name": "NNXRuntimeDml", "skip": true, "comment": "TODO enable, a CL with related fixes is in progress."} ],
 			"targets": [
-				{	"target": "Abs", "additional_datasets": [ "Elementwise.Slow" ]},
+				{	"target": "Abs", "additional_datasets": [ "Elementwise.Slow.OneInput" ]},
 				{	"target": "Acos", "runtimes": [
 						{ "name": "NNXRuntimeORTDml", "absolute_error":  1e-4, "relative_error": 0.03},
 						{ "name": "NNXRuntimeDml",    "absolute_error":  1e-4, "relative_error": 0.03},
@@ -77,7 +78,7 @@
 					]},
 				{	"target": "Acosh" },
 				{	"target": "Asin", "runtimes": [
-						{ "name": "NNXRuntimeORTDml", "absolute_error":  1e-4, "relative_error": 0.03},
+						{ "name": "NNXRuntimeORTDml", "absolute_error":  1e-4, "relative_error": 1.5},
 						{ "name": "NNXRuntimeDml",    "absolute_error":  1e-4, "relative_error": 1.5},
 						{ "name": "NNXRuntimeHlsl",   "absolute_error":  1e-4, "relative_error": 1.5}
 					]},
@@ -92,8 +93,8 @@
 						{ "name": "NNXRuntimeHlsl",   "absolute_error":  1e-4, "relative_error": 0.03}
 					]},
 				{	"target": "Atanh" },
-				{	"target": "BitShift", "skip": true, "comment": "//TODO need attributes" },
-				{	"target": "Cast",     "skip": true, "comment": "//TODO need attributes" },
+				{	"target": "BitShift", "skip": true, "comment": "TODO need attributes" },
+				{	"target": "Cast",     "skip": true, "comment": "TODO need attributes" },
 				{	"target": "Ceil" },
 				{	"target": "Clip", "additional_datasets": [ "Elementwise.ClipTestWithExtraInput" ], "runtimes": [
 						{ "name": "NNXRuntimeDml", "skip": true, "comment": "RDG need scalar tensor support"},
@@ -101,7 +102,7 @@
 					]},
 				{	"target": "Cos" },
 				{	"target": "Cosh" },
-				{	"target": "Elu", "comment": "//TODO test with not-default attribute", "runtimes": [
+				{	"target": "Elu", "comment": "TODO test with not-default attribute", "runtimes": [
 						{ "name": "NNXRuntimeORTDml", "absolute_error":  1e-4, "relative_error": 0.03},
 						{ "name": "NNXRuntimeDml",    "absolute_error":  1e-4, "relative_error": 0.03},
 						{ "name": "NNXRuntimeHlsl",   "absolute_error":  1e-4, "relative_error": 0.03}
@@ -113,11 +114,11 @@
 					]},
 				{	"target": "Exp" },
 				{	"target": "Floor" },
-				{	"target": "IsInf", "skip": true, "comment": "//TODO add flexibility in test data setup + bool tensors" },
-				{	"target": "IsNan", "skip": true, "comment": "//TODO add flexibility in test data setup + bool tensors" },
-				{	"target": "HardSigmoid",         "comment": "//TODO test with not-default attribute" },
-				{	"target": "HardSwish" },
-				{	"target": "LeakyRelu",           "comment": "//TODO test with not-default attribute" },
+				{	"target": "IsInf", "skip": true, "comment": "TODO add flexibility in test data setup + bool tensors" },
+				{	"target": "IsNan", "skip": true, "comment": "TODO add flexibility in test data setup + bool tensors" },
+				{	"target": "HardSigmoid",         "comment": "TODO test with not-default attribute" },
+				{	"target": "HardSwish", "runtimes": [ { "name": "NNXRuntimeDml", "skip": true, "comment": "Not implemented atm on NNXRuntimeDML."} ]},
+				{	"target": "LeakyRelu",           "comment": "TODO test with not-default attribute" },
 				{	"target": "Log", "runtimes": [
 						{ "name": "NNXRuntimeORTDml", "relative_error": 0.15},
 						{ "name": "NNXRuntimeDml",    "relative_error": 0.15},
@@ -131,7 +132,7 @@
 				{	"target": "Reciprocal" },
 				{	"target": "Relu" },
 				{	"target": "Round" },
-				{	"target": "Selu", "comment": "//TODO test with not-default attribute", "runtimes": [
+				{	"target": "Selu", "comment": "TODO test with not-default attribute", "runtimes": [
 						{ "name": "NNXRuntimeORTDml", "absolute_error":  1e-4, "relative_error": 0.03},
 						{ "name": "NNXRuntimeDml",    "absolute_error":  1e-4, "relative_error": 0.03},
 						{ "name": "NNXRuntimeHlsl",   "absolute_error":  1e-4, "relative_error": 0.03}
@@ -158,6 +159,54 @@
 						{ "name": "NNXRuntimeHlsl",   "absolute_error":  1e-4, "relative_error": 0.03}
 					]}
 			]
+		},
+		{
+			"category": "Elementwise.Binary",
+			"runtimes": [ { "name": "NNXRuntimeDml", "skip": true} ],
+			"targets": [
+				{	"target": "Add", "additional_datasets": [ "Elementwise.Slow.TwoInputs" ]},
+				{	"target": "And", "input_type": "boolean", "output_type": "boolean", "runtimes": [
+						{ "name": "NNXRuntimeDml", "skip": true, "comment": "RDG need bool tensor support"},
+						{ "name": "NNXRuntimeHlsl", "skip": true, "comment": "RDG need bool tensor support"}						
+					]},
+				{	"target": "Div"},
+				{	"target": "Equal", "output_type": "boolean", "runtimes": [
+						{ "name": "NNXRuntimeDml", "skip": true, "comment": "RDG need bool tensor support"},
+						{ "name": "NNXRuntimeHlsl", "skip": true, "comment": "RDG need bool tensor support"}						
+					]},
+				{	"target": "Greater", "output_type": "boolean", "runtimes": [
+						{ "name": "NNXRuntimeDml", "skip": true, "comment": "RDG need bool tensor support"},
+						{ "name": "NNXRuntimeHlsl", "skip": true, "comment": "RDG need bool tensor support"}						
+					]},
+				{	"target": "GreaterOrEqual", "output_type": "boolean", "runtimes": [
+						{ "name": "NNXRuntimeDml", "skip": true, "comment": "RDG need bool tensor support"},
+						{ "name": "NNXRuntimeHlsl", "skip": true, "comment": "RDG need bool tensor support"}						
+					]},
+				{	"target": "Less", "output_type": "boolean", "runtimes": [
+						{ "name": "NNXRuntimeDml", "skip": true, "comment": "RDG need bool tensor support"},
+						{ "name": "NNXRuntimeHlsl", "skip": true, "comment": "RDG need bool tensor support"}						
+					]},
+				{	"target": "LessOrEqual", "output_type": "boolean", "runtimes": [
+						{ "name": "NNXRuntimeDml", "skip": true, "comment": "RDG need bool tensor support"},
+						{ "name": "NNXRuntimeHlsl", "skip": true, "comment": "RDG need bool tensor support"}						
+					]},
+				{	"target": "Mod", "skip": true, "comment": "TODO need to set fmod attribute to 1 for floating point tensors"},
+				{	"target": "Mul"},
+				{	"target": "Or", "input_type": "boolean", "output_type": "boolean", "runtimes": [
+						{ "name": "NNXRuntimeDml", "skip": true, "comment": "RDG need bool tensor support"},
+						{ "name": "NNXRuntimeHlsl", "skip": true, "comment": "RDG need bool tensor support"}						
+					]},
+				{	"target": "PRelu", "removed_datasets": [ 
+						"Elementwise.Binary.BidirectionalBroadcast", 
+						"Elementwise.Binary.LHSToRHSBroadcast" ], 
+					"comment": "Prelu does not support bidirectional broadcast. It only support slope(RHS) to input0(LHS)"},
+				{	"target": "Pow", "absolute_error": 1e-4},
+				{	"target": "Sub"},
+				{	"target": "Xor", "input_type": "boolean", "output_type": "boolean", "runtimes": [
+						{ "name": "NNXRuntimeDml", "skip": true, "comment": "RDG need bool tensor support"},
+						{ "name": "NNXRuntimeHlsl", "skip": true, "comment": "RDG need bool tensor support"}						
+					]}
+			]
 		}
 	],
 
@@ -188,13 +237,82 @@
 			]
 		},
 		{
-			"name": "Elementwise.Slow",
+			"name": "Elementwise.Slow.OneInput",
 			"datasets": [
 				{
 					"inputs": [ { "shape": [ 65536, 513 ] } ]
 				}
 			]
 		},
+		{
+			"name": "Elementwise.Binary.NoBroadcast",
+			"datasets": [
+				{ "inputs": [ { "shape": [ 1 ] }, { "shape": [ 1 ] } ]},
+				{ "inputs": [ { "shape": [ 4 ] }, { "shape": [ 4 ] } ]},
+				{ "inputs": [ { "shape": [ 2,3,4,5,6 ] }, { "shape": [ 2,3,4,5,6 ] } ] }
+			]
+		},
+		{
+			"name": "Elementwise.Binary.ZeroSized",
+			"datasets": [
+				{ "inputs": [ { "shape": [ 1,0 ] }, { "shape": [ 1,1 ] } ],
+					"runtimes": [
+						{ "name": "NNXRuntimeHLSL", "skip": true, "comment": "RDG does not support 0 sized tensor" },
+						{ "name": "NNXRuntimeDML",  "skip": true, "comment": "RDG does not support 0 sized tensor" }
+					]}
+			]
+		},
+		{
+			"name": "Elementwise.Binary.RHSToLHSBroadcast",
+			"datasets": [
+				{ "inputs": [ { "shape": [ 2,3,4 ] }, { "shape": [ 1,1,1 ] } ]},
+				{ "inputs": [ { "shape": [ 2,3,4 ] }, { "shape": [ 2,1,1 ] } ]},
+				{ "inputs": [ { "shape": [ 2,3,4 ] }, { "shape": [ 1,1,4 ] } ]},
+				{ "inputs": [ { "shape": [ 2,3,4 ] }, { "shape": [ 1,3,1 ] } ]},
+				{ "inputs": [ { "shape": [ 2,3,4 ] }, { "shape": [ 1,3,4 ] } ]},
+				{ "inputs": [ { "shape": [ 2,3,4 ] }, { "shape": [ 2,1,4 ] } ]},
+				{ "inputs": [ { "shape": [ 2,3,4 ] }, { "shape": [ 1 ]     } ]},
+				{ "inputs": [ { "shape": [ 2,3,4 ] }, { "shape": [ 4 ]     } ]},
+				{ "inputs": [ { "shape": [ 2,3,4 ] }, { "shape": [ 1,1 ]   } ]},
+				{ "inputs": [ { "shape": [ 2,3,4 ] }, { "shape": [ 3,1 ]   } ]},
+				{ "inputs": [ { "shape": [ 2,3,4 ] }, { "shape": [ 1,4 ]   } ]},
+				{ "inputs": [ { "shape": [ 2,3,4 ] }, { "shape": [ 3,4 ]   } ]}
+			]
+		},
+		{
+			"name": "Elementwise.Binary.BidirectionalBroadcast",
+			"datasets": [
+				{ "inputs": [ { "shape": [ 2,3,1 ] }, { "shape": [ 1,1,4 ] } ], "outputs": [ { "shape": [ 2,3,4 ] } ]},
+				{ "inputs": [ { "shape": [ 2,1,4 ] }, { "shape": [ 2,3,1 ] } ], "outputs": [ { "shape": [ 2,3,4 ] } ]},
+				{ "inputs": [ { "shape": [ 2,3,1 ] }, { "shape": [ 4 ]     } ], "outputs": [ { "shape": [ 2,3,4 ] } ]},
+				{ "inputs": [ { "shape": [ 2,3,1 ] }, { "shape": [ 1,4 ]   } ], "outputs": [ { "shape": [ 2,3,4 ] } ]}
+			]
+		},
+		{
+			"name": "Elementwise.Binary.LHSToRHSBroadcast",
+			"datasets": [
+				{ "inputs": [ { "shape": [ 1,1,1 ] }, { "shape": [ 2,3,4 ] } ], "outputs": [ { "shape": [ 2,3,4 ] } ]},
+				{ "inputs": [ { "shape": [ 2,1,1 ] }, { "shape": [ 2,3,4 ] } ], "outputs": [ { "shape": [ 2,3,4 ] } ]},
+				{ "inputs": [ { "shape": [ 1,1,4 ] }, { "shape": [ 2,3,4 ] } ], "outputs": [ { "shape": [ 2,3,4 ] } ]},
+				{ "inputs": [ { "shape": [ 1,3,1 ] }, { "shape": [ 2,3,4 ] } ], "outputs": [ { "shape": [ 2,3,4 ] } ]},
+				{ "inputs": [ { "shape": [ 1,3,4 ] }, { "shape": [ 2,3,4 ] } ], "outputs": [ { "shape": [ 2,3,4 ] } ]},
+				{ "inputs": [ { "shape": [ 2,1,4 ] }, { "shape": [ 2,3,4 ] } ], "outputs": [ { "shape": [ 2,3,4 ] } ]},
+				{ "inputs": [ { "shape": [ 1 ]     }, { "shape": [ 2,3,4 ] } ], "outputs": [ { "shape": [ 2,3,4 ] } ]},
+				{ "inputs": [ { "shape": [ 4 ]     }, { "shape": [ 2,3,4 ] } ], "outputs": [ { "shape": [ 2,3,4 ] } ]},
+				{ "inputs": [ { "shape": [ 1,1 ]   }, { "shape": [ 2,3,4 ] } ], "outputs": [ { "shape": [ 2,3,4 ] } ]},
+				{ "inputs": [ { "shape": [ 3,1 ]   }, { "shape": [ 2,3,4 ] } ], "outputs": [ { "shape": [ 2,3,4 ] } ]},
+				{ "inputs": [ { "shape": [ 1,4 ]   }, { "shape": [ 2,3,4 ] } ], "outputs": [ { "shape": [ 2,3,4 ] } ]},
+				{ "inputs": [ { "shape": [ 3,4 ]   }, { "shape": [ 2,3,4 ] } ], "outputs": [ { "shape": [ 2,3,4 ] } ]}
+			]
+		},
+		{
+			"name": "Elementwise.Slow.TwoInputs",
+			"datasets": [
+				{
+					"inputs": [ { "shape": [ 65536, 513 ] }, { "shape": [ 1 ] } ]
+				}
+			]
+		}
 	]
 	
 }
diff --git a/Engine/Plugins/Experimental/NNX/Source/NNXQA/Resources/NNXQATestDescSchema.json b/Engine/Plugins/Experimental/NNX/Source/NNXQA/Resources/NNXQATestDescSchema.json
--- a/Engine/Plugins/Experimental/NNX/Source/NNXQA/Resources/NNXQATestDescSchema.json
+++ b/Engine/Plugins/Experimental/NNX/Source/NNXQA/Resources/NNXQATestDescSchema.json
@@ -121,6 +121,11 @@
 					"description": "additional dataset associated with this target, by default all dataset witch name are a super set of this category name will be runned, you can specify more dataset using this parameter",
 					"items": { "type": "string" }
 				},
+				"removed_datasets": {
+					"type": "array",
+					"description": "dataset that will be explicitely not run with this target, by default all dataset witch name are a super set of this category name will be runned, you can specify more dataset using this parameter",
+					"items": { "type": "string" }
+				},
 				"skip": {
 					"type": "boolean",
 					"description": "true to skip this operator or model test entirely, false by default"
@@ -178,6 +183,11 @@
 					"type": "boolean",
 					"description": "true if the category contain model tests, false for operators tests"
 				},
+				"runtimes": {
+					"type": "array",
+					"description": "Array of operator to test",
+					"items": { "$ref": "#/$defs/test_config_runtime" }
+				},
 				"comment": {
 					"type": "string",
 					"description": "Unused by test system, usefull to add clarification in json file"
