Fixed crash caused by the constructor of a static test object.

Problem:
Running the packaged NNXIncubator currently crashes, as the constructor of a static object of FParametricTests accesses the uninitialized plugin system.

Solution:
Added an InitializeParametricTests to the NNXQA_API to allow the NNXQAModule to load the json file on startup and also initializing the test automation registry.

#rnx

diff --git a/Engine/Plugins/Experimental/NNX/Source/NNXQA/Private/NNXQAModule.cpp b/Engine/Plugins/Experimental/NNX/Source/NNXQA/Private/NNXQAModule.cpp
--- a/Engine/Plugins/Experimental/NNX/Source/NNXQA/Private/NNXQAModule.cpp
+++ b/Engine/Plugins/Experimental/NNX/Source/NNXQA/Private/NNXQAModule.cpp
@@ -2,4 +2,21 @@
 
 #include "Modules/ModuleManager.h"
 
-IMPLEMENT_MODULE(FDefaultModuleImpl, NNXQA);
\ No newline at end of file
+#include "NNXQAParametricTest.h"
+
+class FNNXQAModule : public IModuleInterface
+{
+public:
+
+	virtual void StartupModule() override
+	{
+		NNX::Test::InitializeParametricTests();
+	}
+
+	virtual void ShutdownModule() override
+	{
+	}
+
+};
+
+IMPLEMENT_MODULE(FNNXQAModule, NNXQA);
\ No newline at end of file
diff --git a/Engine/Plugins/Experimental/NNX/Source/NNXQA/Private/NNXQAParametricTest.cpp b/Engine/Plugins/Experimental/NNX/Source/NNXQA/Private/NNXQAParametricTest.cpp
--- a/Engine/Plugins/Experimental/NNX/Source/NNXQA/Private/NNXQAParametricTest.cpp
+++ b/Engine/Plugins/Experimental/NNX/Source/NNXQA/Private/NNXQAParametricTest.cpp
@@ -26,7 +26,7 @@ namespace Test
 	public:
 		FParametricTests()
 		{
-			ReloadTestDescriptionsFromJson();
+			
 		}
 
 		bool ReloadTestDescriptionsFromJson()
@@ -465,5 +465,14 @@ namespace Test
 		)
 	);
 
+	bool InitializeParametricTests()
+	{
+		bool bResult = ParametricTests.ReloadTestDescriptionsFromJson();
+#if WITH_DEV_AUTOMATION_TESTS
+		ParametricTestAutomationRegistry.Refresh();
+#endif
+		return bResult;
+	}
+
 } // namespace Test
 } // namespace NNX
diff --git a/Engine/Plugins/Experimental/NNX/Source/NNXQA/Public/NNXQAParametricTest.h b/Engine/Plugins/Experimental/NNX/Source/NNXQA/Public/NNXQAParametricTest.h
--- a/Engine/Plugins/Experimental/NNX/Source/NNXQA/Public/NNXQAParametricTest.h
+++ b/Engine/Plugins/Experimental/NNX/Source/NNXQA/Public/NNXQAParametricTest.h
@@ -16,6 +16,7 @@ namespace Test
 
 	//Set the RuntimeFilter witch automation will use, Empty string to run on all runtime (this is default)
 	NNXQA_API void SetAutomationRuntimeFilter(const FString& RuntimeFilter);
-
+	
+	NNXQA_API bool InitializeParametricTests();
 } // namespace Test
 } // namespace NNX
\ No newline at end of file
