NNX: Fixes for element wise unary operators for NNXRuntimeDML

Problem:
* NNXQA Element wise unary operator test are disabled for NNXRuntimeDML atm as some of them are failing (for various reasons see below).

Solution:
* Fix tipos in operator class registration
* Default init attribute for some ops (selu/elu/hardsigmoid/leakyrelu)
* Bump max number of dml binding descriptor (needed as we leak them atm, proper fix will come latter)
* Reactivate unary elem wise test

diff --git a/Engine/Plugins/Experimental/NNX/Source/NNXQA/Resources/NNXQATestDesc.json b/Engine/Plugins/Experimental/NNX/Source/NNXQA/Resources/NNXQATestDesc.json
--- a/Engine/Plugins/Experimental/NNX/Source/NNXQA/Resources/NNXQATestDesc.json
+++ b/Engine/Plugins/Experimental/NNX/Source/NNXQA/Resources/NNXQATestDesc.json
@@ -68,7 +68,6 @@
 	"operator_test_categories": [
 		{
 			"category": "Elementwise.Unary",
-			"runtimes": [ { "name": "NNXRuntimeDml", "skip": true, "comment": "TODO enable, a CL with related fixes is in progress."} ],
 			"targets": [
 				{	"target": "Abs", "additional_datasets": [ "Elementwise.Slow.OneInput" ]},
 				{	"target": "Acos", "runtimes": [
diff --git a/Engine/Plugins/Experimental/NNX/Source/NNXRuntimeRDG/Private/NNXRuntimeDML.cpp b/Engine/Plugins/Experimental/NNX/Source/NNXRuntimeRDG/Private/NNXRuntimeDML.cpp
--- a/Engine/Plugins/Experimental/NNX/Source/NNXRuntimeRDG/Private/NNXRuntimeDML.cpp
+++ b/Engine/Plugins/Experimental/NNX/Source/NNXRuntimeRDG/Private/NNXRuntimeDML.cpp
@@ -165,7 +165,9 @@ OP_EW(SUBTRACT)
 //
 //
 //
-constexpr const uint32 MaxNumDescriptors = 256;
+//HACK: ATM we do not free the descriptors on inference model destruction, so we need to have a big pool until
+//this is fixed. NNXQA Tests will still fail if run repeatedly in the same session until this is fixed.
+constexpr const uint32 MaxNumDescriptors = 4096;
 
 
 //
@@ -370,11 +372,6 @@ template
 >
 class FMLOperatorDmlElementWiseUnary : public FMLOperatorDml
 {
-	float Alpha { 0.0f };
-	float Beta{ 0.0f };
-	float Gamma{ 0.0f };
-	uint32 Num{ 1 };
-
 public:
 
 	static FMLOperatorDml* Create()
@@ -382,9 +379,15 @@ public:
 		return new FMLOperatorDmlElementWiseUnary();
 	}
 
+	virtual ~FMLOperatorDmlElementWiseUnary() = default;
+
 private:
 
-	FMLOperatorDmlElementWiseUnary() = default;
+	FMLOperatorDmlElementWiseUnary() : Alpha(0.0f), Beta(0.0f), Gamma(0.0f), Num(1) {}
+	float Alpha;
+	float Beta;
+	float Gamma;
+	uint32 Num;
 
 public:
 
@@ -444,6 +447,36 @@ private:
 		Desc.Steepness = 1.0f;
 	}
 
+	void InitDmlOpDesc(DML_ACTIVATION_SCALED_ELU_OPERATOR_DESC& Desc, FDmlTensorDesc& TensorDesc)
+	{
+		Desc.InputTensor = &TensorDesc.Desc;
+		Desc.OutputTensor = &TensorDesc.Desc;
+		Desc.Alpha = Alpha;
+		Desc.Gamma = Gamma;
+	}
+
+	void InitDmlOpDesc(DML_ACTIVATION_ELU_OPERATOR_DESC& Desc, FDmlTensorDesc& TensorDesc)
+	{
+		Desc.InputTensor = &TensorDesc.Desc;
+		Desc.OutputTensor = &TensorDesc.Desc;
+		Desc.Alpha = Alpha;
+	}
+
+	void InitDmlOpDesc(DML_ACTIVATION_HARD_SIGMOID_OPERATOR_DESC& Desc, FDmlTensorDesc& TensorDesc)
+	{
+		Desc.InputTensor = &TensorDesc.Desc;
+		Desc.OutputTensor = &TensorDesc.Desc;
+		Desc.Alpha = Alpha;
+		Desc.Beta = Beta;
+	}
+
+	void InitDmlOpDesc(DML_ACTIVATION_LEAKY_RELU_OPERATOR_DESC& Desc, FDmlTensorDesc& TensorDesc)
+	{
+		Desc.InputTensor = &TensorDesc.Desc;
+		Desc.OutputTensor = &TensorDesc.Desc;
+		Desc.Alpha = Alpha;
+	}
+
 public:
 
 	//
@@ -558,6 +591,25 @@ public:
 	}
 };
 
+template<> FMLOperatorDmlElementWiseUnary<DML_ACTIVATION_SCALED_ELU_OPERATOR_DESC, EMLElementWiseUnaryOperatorType::Selu>::FMLOperatorDmlElementWiseUnary()
+	: Alpha(1.67326319217681884765625f), Beta(0.0f), Gamma(1.05070102214813232421875f), Num(1)
+{
+}
+
+template<> FMLOperatorDmlElementWiseUnary<DML_ACTIVATION_ELU_OPERATOR_DESC, EMLElementWiseUnaryOperatorType::Elu>::FMLOperatorDmlElementWiseUnary()
+	: Alpha(1.0f), Beta(0.0f), Gamma(1.05070102214813232421875f), Num(1)
+{
+}
+
+template<> FMLOperatorDmlElementWiseUnary<DML_ACTIVATION_HARD_SIGMOID_OPERATOR_DESC, EMLElementWiseUnaryOperatorType::HardSigmoid>::FMLOperatorDmlElementWiseUnary()
+	: Alpha(0.2f), Beta(0.5f), Gamma(0.0f), Num(1)
+{
+}
+
+template<> FMLOperatorDmlElementWiseUnary<DML_ACTIVATION_LEAKY_RELU_OPERATOR_DESC, EMLElementWiseUnaryOperatorType::LeakyRelu>::FMLOperatorDmlElementWiseUnary()
+	: Alpha(0.01f), Beta(0.0f), Gamma(0.0f), Num(1)
+{
+}
 
 /**
  * Element-wise binary ML operator implementation
@@ -943,16 +995,16 @@ bool FMLRuntimeDml::RegisterElementWiseUnaryOperators()
 
 	OP(DML_ELEMENT_WISE_ABS_OPERATOR_DESC, Abs);
 	OP(DML_ELEMENT_WISE_ACOS_OPERATOR_DESC, Acos);
-	OP(DML_ELEMENT_WISE_ACOS_OPERATOR_DESC, Acosh);
+	OP(DML_ELEMENT_WISE_ACOSH_OPERATOR_DESC, Acosh);
 	OP(DML_ELEMENT_WISE_ASIN_OPERATOR_DESC, Asin);
-	OP(DML_ELEMENT_WISE_ASIN_OPERATOR_DESC, Asinh);
+	OP(DML_ELEMENT_WISE_ASINH_OPERATOR_DESC, Asinh);
 	OP(DML_ELEMENT_WISE_ATAN_OPERATOR_DESC, Atan);
-	OP(DML_ELEMENT_WISE_ATAN_OPERATOR_DESC, Atanh);
+	OP(DML_ELEMENT_WISE_ATANH_OPERATOR_DESC, Atanh);
 	OP(DML_ELEMENT_WISE_CEIL_OPERATOR_DESC, Ceil);
 	OP(DML_ELEMENT_WISE_COS_OPERATOR_DESC, Cos);
 	OP(DML_ELEMENT_WISE_COSH_OPERATOR_DESC, Cosh);
 	OP(DML_ACTIVATION_ELU_OPERATOR_DESC, Elu);
-	OP(DML_ELEMENT_WISE_EXP_OPERATOR_DESC, Erf);
+	OP(DML_ELEMENT_WISE_ERF_OPERATOR_DESC, Erf);
 	OP(DML_ELEMENT_WISE_EXP_OPERATOR_DESC, Exp);
 	OP(DML_ELEMENT_WISE_FLOOR_OPERATOR_DESC, Floor);
 	OP(DML_ELEMENT_WISE_IS_INFINITY_OPERATOR_DESC, IsInf);
