Fixed issue with wrong module name.

Problem:
A copy paste error excluded the hlsl module from the static build causing a crash on startup.

Solution:
Fixed names.

Testing:
Tested by running the standalone NNXIncubator

#rnx

diff --git a/Engine/Plugins/Experimental/NNX/Source/NNXHlslShaders/Private/NNXHlslModule.cpp b/Engine/Plugins/Experimental/NNX/Source/NNXHlslShaders/Private/NNXHlslModule.cpp
--- a/Engine/Plugins/Experimental/NNX/Source/NNXHlslShaders/Private/NNXHlslModule.cpp
+++ b/Engine/Plugins/Experimental/NNX/Source/NNXHlslShaders/Private/NNXHlslModule.cpp
@@ -38,4 +38,4 @@ public:
 	}
 };
 
-IMPLEMENT_MODULE(FNNXHlslShadersModule, NNXHlslShadersModule);
+IMPLEMENT_MODULE(FNNXHlslShadersModule, NNXHlslShaders);
