NNXQA Add variadic tests to parametric tests

Tests: Runned all NNXQA parametric tests in editor win64

diff --git a/Engine/Plugins/Experimental/NNX/Source/NNXQA/Private/NNXQAJsonUtils.h b/Engine/Plugins/Experimental/NNX/Source/NNXQA/Private/NNXQAJsonUtils.h
--- a/Engine/Plugins/Experimental/NNX/Source/NNXQA/Private/NNXQAJsonUtils.h
+++ b/Engine/Plugins/Experimental/NNX/Source/NNXQA/Private/NNXQAJsonUtils.h
@@ -91,6 +91,8 @@ namespace Json
 	struct FTestCategory : FJsonSerializable
 	{
 		FString Category;
+		TArray<FString> AdditionalDatasets;
+		TArray<FString> RemovedDatasets;
 		TArray<FTestConfigTarget> Targets;
 		TArray<FTestConfigRuntime> Runtimes;
 		bool Skip;
@@ -100,6 +102,8 @@ namespace Json
 			JSON_SERIALIZE("category", Category);
 			JSON_SERIALIZE("skip", Skip);
 			JSON_SERIALIZE("is_model_test", IsModelTest);
+			JSON_SERIALIZE_ARRAY("additional_datasets", AdditionalDatasets);
+			JSON_SERIALIZE_ARRAY("removed_datasets", RemovedDatasets);
 			JSON_SERIALIZE_ARRAY_SERIALIZABLE("targets", Targets, FTestConfigTarget);
 			JSON_SERIALIZE_ARRAY_SERIALIZABLE("runtimes", Runtimes, FTestConfigRuntime);
 		END_JSON_SERIALIZER
diff --git a/Engine/Plugins/Experimental/NNX/Source/NNXQA/Private/NNXQAParametricTest.cpp b/Engine/Plugins/Experimental/NNX/Source/NNXQA/Private/NNXQAParametricTest.cpp
--- a/Engine/Plugins/Experimental/NNX/Source/NNXQA/Private/NNXQAParametricTest.cpp
+++ b/Engine/Plugins/Experimental/NNX/Source/NNXQA/Private/NNXQAParametricTest.cpp
@@ -194,6 +194,18 @@ namespace Test
 			}
 			return TestSuffix;
 		}
+
+		static bool IsSubstringFoundInArray(const TArray<FString>& Names, const FString& SubString)
+		{
+			for (auto&& Name : Names)
+			{
+				if (Name.Contains(SubString))
+				{
+					return true;
+				}
+			}
+			return false;
+		}
 		
 		void AddTestFromCategory(const FString& BaseTestPath, const TArray<Json::FTestCategory>& TestCategories, const TArray<Json::FTestConfigInputOutputSet> InputOutputSets)
 		{
@@ -221,10 +233,14 @@ namespace Test
 					
 					for (const Json::FTestConfigInputOutputSet& InputOutputSet : InputOutputSets)
 					{
-						//If category is a substring of InputOutputSet name or if target explicitly requested the InputOutputSet name
-						if (InputOutputSet.Name.Contains(TestCategory.Category) || TestTarget.AdditionalDatasets.Contains(InputOutputSet.Name))
+						//Accept test if category is a substring of the dataset name or one of the explicitly requested dataset names
+						if (InputOutputSet.Name.Contains(TestCategory.Category) || 
+							IsSubstringFoundInArray(TestTarget.AdditionalDatasets, InputOutputSet.Name) ||
+							IsSubstringFoundInArray(TestCategory.AdditionalDatasets, InputOutputSet.Name))
 						{
-							if (TestTarget.RemovedDatasets.Contains(InputOutputSet.Name))
+							//Reject test if dataset is explicitly rejected
+							if (IsSubstringFoundInArray(TestTarget.RemovedDatasets, InputOutputSet.Name) || 
+								IsSubstringFoundInArray(TestCategory.RemovedDatasets, InputOutputSet.Name))
 							{
 								continue;
 							}
diff --git a/Engine/Plugins/Experimental/NNX/Source/NNXQA/Resources/NNXQATestDesc.json b/Engine/Plugins/Experimental/NNX/Source/NNXQA/Resources/NNXQATestDesc.json
--- a/Engine/Plugins/Experimental/NNX/Source/NNXQA/Resources/NNXQATestDesc.json
+++ b/Engine/Plugins/Experimental/NNX/Source/NNXQA/Resources/NNXQATestDesc.json
@@ -206,6 +206,17 @@
 						{ "name": "NNXRuntimeHlsl", "skip": true, "comment": "RDG need bool tensor support"}						
 					]}
 			]
+		},
+		{
+			"category": "Elementwise.Variadic",
+			"additional_datasets": [ "Elementwise.Binary" ],
+			"runtimes": [ { "name": "NNXRuntimeDml", "skip": true} ],
+			"targets": [
+				{	"target": "Max", "additional_datasets": [ "Elementwise.Slow.TwoInputs" ]},
+				{	"target": "Mean"},
+				{	"target": "Min"},
+				{	"target": "Sum"}
+			]
 		}
 	],
 
@@ -251,6 +262,15 @@
 				{ "inputs": [ { "shape": [ 2,3,4,5,6 ] }, { "shape": [ 2,3,4,5,6 ] } ] }
 			]
 		},
+		{
+			"name": "Elementwise.Variadic.NoBroadcast",
+			"datasets": [
+				{ "inputs": [ { "shape": [ 1 ] }, { "shape": [ 1 ] }, { "shape": [ 1 ] } ]},
+				{ "inputs": [ { "shape": [ 1 ] }, { "shape": [ 1 ] }, { "shape": [ 1 ] }, { "shape": [ 1 ] } ]},
+				{ "inputs": [ { "shape": [ 4 ] }, { "shape": [ 4 ] }, { "shape": [ 4 ] } ]},
+				{ "inputs": [ { "shape": [ 5 ] }, { "shape": [ 5 ] }, { "shape": [ 5 ] }, { "shape": [ 5 ] }, { "shape": [ 5 ] }, { "shape": [ 5 ] }, { "shape": [ 5 ] }, { "shape": [ 5 ] } ]}
+			]
+		},
 		{
 			"name": "Elementwise.Binary.ZeroSized",
 			"datasets": [
@@ -287,6 +307,15 @@
 				{ "inputs": [ { "shape": [ 2,3,1 ] }, { "shape": [ 1,4 ]   } ], "outputs": [ { "shape": [ 2,3,4 ] } ]}
 			]
 		},
+		{
+			"name": "Elementwise.Variadic.MultiDirectionalBroadcast",
+			"datasets": [
+				{ "inputs": [ { "shape": [ 2,3,1 ] }, { "shape": [ 1,1,4 ] }, { "shape": [ 1,1,1 ] } ], "outputs": [ { "shape": [ 2,3,4 ] } ]},
+				{ "inputs": [ { "shape": [ 2,1,4 ] }, { "shape": [ 1,3,4 ] }, { "shape": [ 2,3,1 ] } ], "outputs": [ { "shape": [ 2,3,4 ] } ]},
+				{ "inputs": [ { "shape": [ 2,3,1 ] }, { "shape": [ 4 ]     }, { "shape": [ 3,1 ]   } ], "outputs": [ { "shape": [ 2,3,4 ] } ]},
+				{ "inputs": [ { "shape": [ 2,3,1 ] }, { "shape": [ 3,1 ]   }, { "shape": [ 4 ]     } ], "outputs": [ { "shape": [ 2,3,4 ] } ]}
+			]
+		},
 		{
 			"name": "Elementwise.Binary.LHSToRHSBroadcast",
 			"datasets": [
diff --git a/Engine/Plugins/Experimental/NNX/Source/NNXQA/Resources/NNXQATestDescSchema.json b/Engine/Plugins/Experimental/NNX/Source/NNXQA/Resources/NNXQATestDescSchema.json
--- a/Engine/Plugins/Experimental/NNX/Source/NNXQA/Resources/NNXQATestDescSchema.json
+++ b/Engine/Plugins/Experimental/NNX/Source/NNXQA/Resources/NNXQATestDescSchema.json
@@ -118,12 +118,12 @@
 				},
 				"additional_datasets": {
 					"type": "array",
-					"description": "additional dataset associated with this target, by default all dataset witch name are a super set of this category name will be runned, you can specify more dataset using this parameter",
+					"description": "additional dataset associated with this target, by default all dataset witch name are a super set of this target category name will be runned, you can specify more dataset using this parameter",
 					"items": { "type": "string" }
 				},
 				"removed_datasets": {
 					"type": "array",
-					"description": "dataset that will be explicitely not run with this target, by default all dataset witch name are a super set of this category name will be runned, you can specify more dataset using this parameter",
+					"description": "dataset that will be explicitely not run with this target, by default all dataset witch name are a super set of this target category name will be runned, you can filter out dataset using this parameter",
 					"items": { "type": "string" }
 				},
 				"skip": {
@@ -188,6 +188,16 @@
 					"description": "Array of operator to test",
 					"items": { "$ref": "#/$defs/test_config_runtime" }
 				},
+				"additional_datasets": {
+					"type": "array",
+					"description": "additional dataset associated with this category, by default all dataset witch name are a super set of this category name will be runned, you can specify more dataset using this parameter",
+					"items": { "type": "string" }
+				},
+				"removed_datasets": {
+					"type": "array",
+					"description": "dataset that will be explicitely not run with this category, by default all dataset witch name are a super set of this category name will be runned, you can filter out dataset using this parameter",
+					"items": { "type": "string" }
+				},
 				"comment": {
 					"type": "string",
 					"description": "Unused by test system, usefull to add clarification in json file"
