NNXQA Cleanup, remove all non parametric test and related code now that they have been ported.

Tests: Verify all parametric tests are still present + run them in editor.

diff --git a/Engine/Plugins/Experimental/NNX/Source/NNXQA/Private/NNXQAOperatorTest.cpp b/Engine/Plugins/Experimental/NNX/Source/NNXQA/Private/NNXQAOperatorTest.cpp
deleted file mode 100644
diff --git a/Engine/Plugins/Experimental/NNX/Source/NNXQA/Private/NNXQAParametricTest.cpp b/Engine/Plugins/Experimental/NNX/Source/NNXQA/Private/NNXQAParametricTest.cpp
--- a/Engine/Plugins/Experimental/NNX/Source/NNXQA/Private/NNXQAParametricTest.cpp
+++ b/Engine/Plugins/Experimental/NNX/Source/NNXQA/Private/NNXQAParametricTest.cpp
@@ -54,7 +54,7 @@ namespace Test
 				return false;
 			}
 
-			const FString NNXBaseTestPath(TEXT("System.Engine.MachineLearning.NNX.Parametric"));
+			const FString NNXBaseTestPath(TEXT("System.Engine.MachineLearning.NNX"));
 			AddTestFromCategory(NNXBaseTestPath + TEXT(".Model."), ModelTestCategories, InputOutputSets);
 			AddTestFromCategory(NNXBaseTestPath + TEXT(".Operator."), OperatorTestCategories, InputOutputSets);
 
@@ -175,21 +175,21 @@ namespace Test
 			for (const Json::FTestConfigTensor& Input : Dataset.Inputs)
 			{
 				if (!bIsFirstShape) TestSuffix += TEXT("_");
-				TestSuffix += ShapeToString(Input.Shape);
+				TestSuffix += ShapeToString<int32>(Input.Shape);
 				bIsFirstShape = false;
 			}
 			TestSuffix += TEXT("=>");
 			//If output is not defined it is the first input shape.
 			if (Dataset.Outputs.Num() == 0)
 			{
-				TestSuffix += ShapeToString(Dataset.Inputs[0].Shape);
+				TestSuffix += ShapeToString<int32>(Dataset.Inputs[0].Shape);
 			}
 			else
 			{
 				for (const Json::FTestConfigTensor& Output : Dataset.Outputs)
 				{
 					TestSuffix += TEXT("_");
-					TestSuffix += ShapeToString(Output.Shape);
+					TestSuffix += ShapeToString<int32>(Output.Shape);
 				}
 			}
 			return TestSuffix;
@@ -310,7 +310,7 @@ namespace Test
 			}
 		}
 
-		return CompareONNXModelInferenceAcrossRuntimes(TestSetup.TestName, ModelData, &TestSetup, RuntimeFilter);
+		return CompareONNXModelInferenceAcrossRuntimes(ModelData, TestSetup, RuntimeFilter);
 	}
 
 	static FString AutomationRuntimeFilter;
diff --git a/Engine/Plugins/Experimental/NNX/Source/NNXQA/Private/NNXQATestsOperatorElementWiseBinary.cpp b/Engine/Plugins/Experimental/NNX/Source/NNXQA/Private/NNXQATestsOperatorElementWiseBinary.cpp
deleted file mode 100644
diff --git a/Engine/Plugins/Experimental/NNX/Source/NNXQA/Private/NNXQATestsOperatorElementWiseBinary.h b/Engine/Plugins/Experimental/NNX/Source/NNXQA/Private/NNXQATestsOperatorElementWiseBinary.h
deleted file mode 100644
diff --git a/Engine/Plugins/Experimental/NNX/Source/NNXQA/Private/NNXQATestsOperatorElementWiseUnary.cpp b/Engine/Plugins/Experimental/NNX/Source/NNXQA/Private/NNXQATestsOperatorElementWiseUnary.cpp
deleted file mode 100644
diff --git a/Engine/Plugins/Experimental/NNX/Source/NNXQA/Private/NNXQATestsOperatorElementWiseUnary.h b/Engine/Plugins/Experimental/NNX/Source/NNXQA/Private/NNXQATestsOperatorElementWiseUnary.h
deleted file mode 100644
diff --git a/Engine/Plugins/Experimental/NNX/Source/NNXQA/Private/NNXQATestsOperatorElementWiseVariadic.cpp b/Engine/Plugins/Experimental/NNX/Source/NNXQA/Private/NNXQATestsOperatorElementWiseVariadic.cpp
deleted file mode 100644
diff --git a/Engine/Plugins/Experimental/NNX/Source/NNXQA/Private/NNXQATestsOperatorElementWiseVariadic.h b/Engine/Plugins/Experimental/NNX/Source/NNXQA/Private/NNXQATestsOperatorElementWiseVariadic.h
deleted file mode 100644
diff --git a/Engine/Plugins/Experimental/NNX/Source/NNXQA/Private/NNXQAUtils.cpp b/Engine/Plugins/Experimental/NNX/Source/NNXQA/Private/NNXQAUtils.cpp
--- a/Engine/Plugins/Experimental/NNX/Source/NNXQA/Private/NNXQAUtils.cpp
+++ b/Engine/Plugins/Experimental/NNX/Source/NNXQA/Private/NNXQAUtils.cpp
@@ -68,26 +68,11 @@ namespace Test
 		check(TensorsDesc.Num() == OutBindings.Num());
 	}
 
-
-	FString ShapeToString(TArrayView<const uint32> Shape)
-	{
-		FString TestSuffix(TEXT("["));
-		bool bIsFirstDim = true;
-		for (uint32 size : Shape)
-		{
-			if (!bIsFirstDim) TestSuffix += TEXT(",");
-			TestSuffix += FString::Printf(TEXT("%d"), size);
-			bIsFirstDim = false;
-		}
-		TestSuffix += TEXT("]");
-		return TestSuffix;
-	}
-
-	FString ShapeToString(TArrayView<const int32> Shape)
+	template<typename T> FString ShapeToString(TArrayView<const T> Shape)
 	{
 		FString TestSuffix(TEXT("["));
 		bool bIsFirstDim = true;
-		for (uint32 size : Shape)
+		for (T size : Shape)
 		{
 			if (!bIsFirstDim) TestSuffix += TEXT(",");
 			TestSuffix += FString::Printf(TEXT("%d"), size);
@@ -96,6 +81,8 @@ namespace Test
 		TestSuffix += TEXT("]");
 		return TestSuffix;
 	}
+	template FString ShapeToString<int32>(TArrayView<const int32> Shape);
+	template FString ShapeToString<uint32>(TArrayView<const uint32> Shape);
 
 	FString FMLTensorDescToString(const FMLTensorDesc& desc)
 	{
@@ -376,15 +363,15 @@ namespace Test
 		return ReturnValue;
 	}
 
-	bool CompareONNXModelInferenceAcrossRuntimes(const FString& Name, TArrayView<uint8> ONNXModelData, const FTests::FTestSetup* AutomationTestSetup, const FString& RuntimeFilter)
+	bool CompareONNXModelInferenceAcrossRuntimes(TArrayView<uint8> ONNXModelData, const FTests::FTestSetup& TestSetup, const FString& RuntimeFilter)
 	{
 		FString CurrentPlatform = UGameplayStatics::GetPlatformName();
-		if (AutomationTestSetup != nullptr && AutomationTestSetup->AutomationExcludedPlatform.Contains(CurrentPlatform))
+		if (TestSetup.AutomationExcludedPlatform.Contains(CurrentPlatform))
 		{
-			UE_LOG(LogNNX, Display, TEXT("Skipping test of '%s' for platform %s (by config)"), *Name, *CurrentPlatform);
+			UE_LOG(LogNNX, Display, TEXT("Skipping test of '%s' for platform %s (by config)"), *TestSetup.TargetName, *CurrentPlatform);
 			return true;
 		}
-		UE_LOG(LogNNX, Display, TEXT("Starting tests of '%s'"), *Name);
+		UE_LOG(LogNNX, Display, TEXT("Starting tests of '%s'"), *TestSetup.TargetName);
 
 		// Reference runtime
 		NNX::IRuntime* RefRuntime = NNX::GetRuntime(TEXT("NNXRuntimeORTCpu"));
@@ -425,9 +412,8 @@ namespace Test
 
 			FString TestResult;
 
-			if (AutomationTestSetup != nullptr && (
-				AutomationTestSetup->AutomationExcludedRuntime.Contains(RuntimeName) ||
-				AutomationTestSetup->AutomationExcludedPlatformRuntimeCombination.Contains(TPair<FString, FString>(CurrentPlatform, RuntimeName))))
+			if (TestSetup.AutomationExcludedRuntime.Contains(RuntimeName) ||
+				TestSetup.AutomationExcludedPlatformRuntimeCombination.Contains(TPair<FString, FString>(CurrentPlatform, RuntimeName)))
 			{
 				TestResult = TEXT("skipped (by config)");
 			}
@@ -436,8 +422,8 @@ namespace Test
 				TArray<TArray<char>> OutputMemBuffers;
 				TArray<NNX::FMLTensorDesc> OutputTensorDescs;
 				bool bTestSuceeded = true;
-				float AbsoluteErrorEpsilon = (AutomationTestSetup != nullptr) ? AutomationTestSetup->GetAbsoluteErrorEpsilonForRuntime(RuntimeName) : 1e-5f;
-				float RelativeErrorPercent = (AutomationTestSetup != nullptr) ? AutomationTestSetup->GetRelativeErrorPercentForRuntime(RuntimeName) : 1e-3f;
+				float AbsoluteErrorEpsilon = TestSetup.GetAbsoluteErrorEpsilonForRuntime(RuntimeName);
+				float RelativeErrorPercent = TestSetup.GetRelativeErrorPercentForRuntime(RuntimeName);
 
 				RunTestInference(ONNXModelData, Runtime, OutputTensorDescs, OutputMemBuffers);
 				if (OutputTensorDescs.Num() == RefOutputTensorDescs.Num())
@@ -473,11 +459,5 @@ namespace Test
 		return TestSetups.Last(0);
 	}
 
-	FTests::FTestSetup& FTests::AddTest(const FString& ModelOrOperatorName, const FString& TestSuffix)
-	{
-		//TODO cleanup: always use the 3 parameter path.
-		return AddTest(TEXT(""), ModelOrOperatorName, TestSuffix);
-	}
-
 } // namespace Test
 } // namespace NNX
\ No newline at end of file
diff --git a/Engine/Plugins/Experimental/NNX/Source/NNXQA/Private/NNXQAUtils.h b/Engine/Plugins/Experimental/NNX/Source/NNXQA/Private/NNXQAUtils.h
--- a/Engine/Plugins/Experimental/NNX/Source/NNXQA/Private/NNXQAUtils.h
+++ b/Engine/Plugins/Experimental/NNX/Source/NNXQA/Private/NNXQAUtils.h
@@ -25,14 +25,6 @@ namespace Test
 				IsModelTest(false)
 			{}
 			
-			FTestSetup(const FString& ModelOrOperatorName, const FString& TestSuffix) :
-				TestName(ModelOrOperatorName+TestSuffix),
-				TargetName(ModelOrOperatorName),
-				AbsoluteErrorEpsilon(DefaultAbsoluteErrorEpsilon),
-				RelativeErrorPercent(DefaultRelativeErrorPercent),
-				IsModelTest(false)
-			{}
-
 			const FString TestName;
 			const FString TargetName;
 			float AbsoluteErrorEpsilon;
@@ -62,17 +54,15 @@ namespace Test
 
 
 		FTestSetup& AddTest(const FString& Category, const FString& ModelOrOperatorName, const FString& TestSuffix);
-		FTestSetup& AddTest(const FString& ModelOrOperatorName, const FString& TestSuffix);
 
 		TArray<FTestSetup> TestSetups;
 	};
 
-	bool CompareONNXModelInferenceAcrossRuntimes(const FString& Name, TArrayView<uint8> ModelData, const FTests::FTestSetup* AutomationTestSetup, const FString& RuntimeFilter = TEXT(""));
+	bool CompareONNXModelInferenceAcrossRuntimes(TArrayView<uint8> ModelData, const FTests::FTestSetup& TestSetup, const FString& RuntimeFilter = TEXT(""));
 
 	FString FMLTensorDescToString(const FMLTensorDesc& TensorDesc);
 	FString TensorToString(const FMLTensorDesc& TensorDesc, const TArray<char>& TensorData);
-	FString ShapeToString(TArrayView<const uint32> Shape);
-	FString ShapeToString(TArrayView<const int32> Shape);
+	template<typename T> FString ShapeToString(TArrayView<const T> Shape);
 
 } // namespace Test
 } // namespace NNX
\ No newline at end of file
