NNXQA Added Smoketest + cleanup in NNXQAParametricTest.cpp

Problem:
We didn't had a very quick set of unit test to run as smoke presubmit tests.

Solution:
Added a command to run them + select as subset of tests runtimes combination running under 5s.

Tests:
* Runned various combination of test commands to test the different behaviour and potential workflow around TDD we can have to develop operators or runtime (focus on a given op, a category of them or a runtime)

* Runned wall time test of different tests selection ->
All 610 test via session frontend / automation all runtimes (but NNXRuntimeORTCuda as we dont have the Dlls atm) : 113s
=> just run all via session frontend

All 610 test via session frontend / automation only NNXRuntimeHLSL : 36s
=> nnx.test.setautomationfilter NNXRuntimeHLSL + run all via session frontend

All 610 tests command line only NNXRuntimeHLSL : 20s
=> nnx.test.run -runtime NNXRuntimeHLSL

All 610 tests command line only NNXRuntimeCPU : 80s
=> nnx.test.run -runtime NNXRuntimeCPU

All 610 tests command line all runtimes (but NNXRuntimeORTCuda as we dont have the Dlls atm) : 106s
=> nnx.test.run

Smoke test (178 tested selected for speed and relevance, all runtime but NNXRuntimeCPU): 5s
=> nnx.test.smokerun

diff --git a/Engine/Plugins/Experimental/NNX/Source/NNXQA/Private/NNXQAParametricTest.cpp b/Engine/Plugins/Experimental/NNX/Source/NNXQA/Private/NNXQAParametricTest.cpp
--- a/Engine/Plugins/Experimental/NNX/Source/NNXQA/Private/NNXQAParametricTest.cpp
+++ b/Engine/Plugins/Experimental/NNX/Source/NNXQA/Private/NNXQAParametricTest.cpp
@@ -320,22 +320,20 @@ namespace Test
 	}
 
 	static FAutoConsoleCommand SetAutomationRuntimeFilterCommand(
-		TEXT("nnx.test.parametric.setautomationfilter"), TEXT("Set the RuntimeFilter witch automation will use, no parameter to run on all runtime (this is the default)."),
+		TEXT("nnx.test.setautomationfilter"), TEXT("Set the RuntimeFilter witch automation will use, no parameter to run on all runtime (this is the default)."),
 		FConsoleCommandWithArgsDelegate::CreateStatic(
 			[](const TArray< FString >& Args)
 			{
-				if (Args.Num() > 0)
+				AutomationRuntimeFilter.Empty();
+				for (auto&& Arg : Args)
 				{
-					SetAutomationRuntimeFilter(Args[0]);
-				}
-				else
-				{
-					SetAutomationRuntimeFilter(TEXT(""));
+					AutomationRuntimeFilter += Arg;
+					AutomationRuntimeFilter += TEXT(" ");
 				}
 			}
 		)
 	);
-	
+
 	bool RunParametricTests(const FString& NameSubstring, const FString& Tag, const FString& InRuntimeFilter)
 	{
 		bool bAllTestSucceed = true;
@@ -362,7 +360,7 @@ namespace Test
 
 		if (NumTest == 0)
 		{
-			UE_LOG(LogNNX, Display, TEXT("No test selected to run (on %d parametric tests registered)."), ParametricTests.TestSetups.Num());
+			UE_LOG(LogNNX, Display, TEXT("No test selected to run (on %d tests registered)."), ParametricTests.TestSetups.Num());
 			return true;
 		}
 		else if (NumTestFailed == 0)
@@ -393,14 +391,54 @@ namespace Test
 	}
 
 	static FAutoConsoleCommand RunTestCommand(
-		TEXT("nnx.test.parametric.run"), TEXT("Run all parametric tests. Use -name to filter by name (substring). Use -tag to filter by tag. Use -runtime to only run for the provided runtime."),
+		TEXT("nnx.test.run"), TEXT("Run all tests. Use -name to filter by name (substring). Use -tag to filter by tag. Use -runtime to only run for the provided runtime, default is to use filter set from setruntimefilter command."),
 		FConsoleCommandWithArgsDelegate::CreateStatic(
 			[](const TArray< FString >& Args)
 			{
 				FString Name = FindArg(TEXT("-name"), Args);
 				FString Tag = FindArg(TEXT("-tag"), Args);
 				FString Runtime = FindArg(TEXT("-runtime"), Args);
-				RunParametricTests(Name, Tag, Runtime);
+				double StartTime = FPlatformTime::Seconds();
+				bool bTestSucceeded = RunParametricTests(Name, Tag, Runtime);
+				double EndTime = FPlatformTime::Seconds();
+				double TimeForTest = static_cast<float>(EndTime - StartTime);
+				if (bTestSucceeded)
+				{
+					UE_LOG(LogNNX, Display, TEXT("Tests succeeded in %0.1f seconds."), TimeForTest);
+				}
+				else
+				{
+					UE_LOG(LogNNX, Warning, TEXT("Tests FAILED in %0.1f seconds."), TimeForTest);
+				}
+			}
+		)
+	);
+
+	static FAutoConsoleCommand RunSmokeTestCommand(
+		TEXT("nnx.test.smokerun"), TEXT("Run all smoke tests. Use -name to additionaly filter by name. Use -runtime to only run for the provided runtime, default is to run on all runtime but NNXRuntimeCPU (too slow at the moment, see comment in code)."),
+		FConsoleCommandWithArgsDelegate::CreateStatic(
+			[](const TArray< FString >& Args)
+			{
+				FString Name = FindArg(TEXT("-name"), Args);
+				FString Runtime = FindArg(TEXT("-runtime"), Args);
+				if (Runtime.IsEmpty())
+				{
+					//NNXRuntimeCPU tests are currently slow due to a synchronization delay on thread creation/destruction of the ORT session.
+					//We want very fast smoke test, thus by default we do NOT run for NNXRuntimeCPU for smoke tests.
+					Runtime = TEXT("NNXRuntimeORTDml NNXRuntimeORTCuda NNXRuntimeHlsl NNXRuntimeDml");
+				}
+				double StartTime = FPlatformTime::Seconds();
+				bool bTestSucceeded = RunParametricTests(Name, TEXT("smoketest"), Runtime);
+				double EndTime = FPlatformTime::Seconds();
+				double TimeForTest = static_cast<float>(EndTime - StartTime);
+				if (bTestSucceeded)
+				{
+					UE_LOG(LogNNX, Display, TEXT("Smoke tests succeeded in %0.1f seconds."), TimeForTest);
+				}
+				else
+				{
+					UE_LOG(LogNNX, Warning, TEXT("Smoke tests FAILED in %0.1f seconds."), TimeForTest);
+				}
 			}
 		)
 	);
@@ -468,7 +506,7 @@ namespace Test
 #endif //WITH_DEV_AUTOMATION_TESTS
 
 	static FAutoConsoleCommand TestReloadCommand(
-		TEXT("nnx.test.parametric.reload"), TEXT("Reload NNX parametric tests definition from Json."),
+		TEXT("nnx.test.reload"), TEXT("Reload NNX tests definition from Json."),
 		FConsoleCommandWithArgsDelegate::CreateStatic(
 			[](const TArray< FString >& Args)
 			{
diff --git a/Engine/Plugins/Experimental/NNX/Source/NNXQA/Private/NNXQAUtils.cpp b/Engine/Plugins/Experimental/NNX/Source/NNXQA/Private/NNXQAUtils.cpp
--- a/Engine/Plugins/Experimental/NNX/Source/NNXQA/Private/NNXQAUtils.cpp
+++ b/Engine/Plugins/Experimental/NNX/Source/NNXQA/Private/NNXQAUtils.cpp
@@ -399,7 +399,7 @@ namespace Test
 				continue;
 			}
 
-			if (!RuntimeFilter.IsEmpty() && RuntimeName != RuntimeFilter)
+			if (!RuntimeFilter.IsEmpty() && !RuntimeFilter.Contains(RuntimeName))
 			{
 				continue;
 			}
diff --git a/Engine/Plugins/Experimental/NNX/Source/NNXQA/Resources/NNXQATestDesc.json b/Engine/Plugins/Experimental/NNX/Source/NNXQA/Resources/NNXQATestDesc.json
--- a/Engine/Plugins/Experimental/NNX/Source/NNXQA/Resources/NNXQATestDesc.json
+++ b/Engine/Plugins/Experimental/NNX/Source/NNXQA/Resources/NNXQATestDesc.json
@@ -52,7 +52,6 @@
 				},
 				{
 					"target": "VertexDeltaModel",
-					"tags": [ "SmokeTest" ],
 					"runtimes": [
 						{ "name": "NNXRuntimeORTDml", "absolute_error": 1e-4, "relative_error": 10.0,
 							"comment": "TODO check, relative error is very high."
@@ -70,7 +69,7 @@
 			"category": "Elementwise.Unary",
 			"targets": [
 				{	"target": "Abs", "additional_datasets": [ "Elementwise.Slow.OneInput" ]},
-				{	"target": "Acos", "runtimes": [
+				{	"target": "Acos", "tags": [ "SmokeTest" ], "runtimes": [
 						{ "name": "NNXRuntimeORTDml", "absolute_error":  1e-4, "relative_error": 0.03},
 						{ "name": "NNXRuntimeDml",    "absolute_error":  1e-4, "relative_error": 0.03},
 						{ "name": "NNXRuntimeHlsl",   "absolute_error":  1e-4, "relative_error": 0.03}
@@ -94,8 +93,8 @@
 				{	"target": "Atanh" },
 				{	"target": "BitShift", "skip": true, "comment": "TODO need attributes" },
 				{	"target": "Cast",     "skip": true, "comment": "TODO need attributes" },
-				{	"target": "Ceil" },
-				{	"target": "Clip", "additional_datasets": [ "Elementwise.ClipTestWithExtraInput" ], "runtimes": [
+				{	"target": "Ceil", "tags": [ "SmokeTest" ]},
+				{	"target": "Clip", "tags": [ "SmokeTest" ], "additional_datasets": [ "Elementwise.ClipTestWithExtraInput" ], "runtimes": [
 						{ "name": "NNXRuntimeDml", "skip": true, "comment": "RDG need scalar tensor support"},
 						{ "name": "NNXRuntimeHlsl", "skip": true, "comment": "RDG need scalar tensor support"}						
 					]},
@@ -106,7 +105,7 @@
 						{ "name": "NNXRuntimeDml",    "absolute_error":  1e-4, "relative_error": 0.03},
 						{ "name": "NNXRuntimeHlsl",   "absolute_error":  1e-4, "relative_error": 0.03}
 					]},
-				{	"target": "Erf", "runtimes": [
+				{	"target": "Erf", "tags": [ "SmokeTest" ], "runtimes": [
 						{ "name": "NNXRuntimeORTDml", "absolute_error":  1e-4, "relative_error": 0.03},
 						{ "name": "NNXRuntimeDml",    "absolute_error":  1e-4, "relative_error": 0.03},
 						{ "name": "NNXRuntimeHlsl",   "absolute_error":  5e-4, "relative_error": 0.05, "comment": "NNXRuntimeHLSL Erf implementation is less precise than the DML one"}
@@ -128,15 +127,15 @@
 						{ "name": "NNXRuntimeDml", "skip": true, "comment": "RDG need bool tensor support"},
 						{ "name": "NNXRuntimeHlsl", "skip": true, "comment": "RDG need bool tensor support"}						
 					]},
-				{	"target": "Reciprocal" },
-				{	"target": "Relu" },
-				{	"target": "Round" },
+				{	"target": "Reciprocal", "tags": [ "SmokeTest" ] },
+				{	"target": "Relu", "tags": [ "SmokeTest" ] },
+				{	"target": "Round", "tags": [ "SmokeTest" ] },
 				{	"target": "Selu", "comment": "TODO test with not-default attribute", "runtimes": [
 						{ "name": "NNXRuntimeORTDml", "absolute_error":  1e-4, "relative_error": 0.03},
 						{ "name": "NNXRuntimeDml",    "absolute_error":  1e-4, "relative_error": 0.03},
 						{ "name": "NNXRuntimeHlsl",   "absolute_error":  1e-4, "relative_error": 0.03}
 					]},
-				{	"target": "Sigmoid" },
+				{	"target": "Sigmoid", "tags": [ "SmokeTest" ] },
 				{	"target": "Sign" },
 				{	"target": "Sin", "runtimes": [
 						{ "name": "NNXRuntimeORTDml", "absolute_error":  1e-4, "relative_error": 0.03},
@@ -144,9 +143,9 @@
 						{ "name": "NNXRuntimeHlsl",   "absolute_error":  1e-4, "relative_error": 0.03}
 					]},
 				{	"target": "Sinh" },
-				{	"target": "Softplus" },
-				{	"target": "Softsign" },
-				{	"target": "Sqrt" },
+				{	"target": "Softplus", "tags": [ "SmokeTest" ] },
+				{	"target": "Softsign", "tags": [ "SmokeTest" ] },
+				{	"target": "Sqrt", "tags": [ "SmokeTest" ] },
 				{	"target": "Tan", "runtimes": [
 						{ "name": "NNXRuntimeORTDml", "absolute_error":  1e-4, "relative_error": 0.03},
 						{ "name": "NNXRuntimeDml",    "absolute_error":  1e-4, "relative_error": 0.03},
@@ -164,16 +163,16 @@
 			"runtimes": [ { "name": "NNXRuntimeDml", "skip": true} ],
 			"targets": [
 				{	"target": "Add", "additional_datasets": [ "Elementwise.Slow.TwoInputs" ]},
-				{	"target": "And", "input_type": "boolean", "output_type": "boolean", "runtimes": [
+				{	"target": "And", "tags": [ "SmokeTest" ], "input_type": "boolean", "output_type": "boolean", "runtimes": [
 						{ "name": "NNXRuntimeDml", "skip": true, "comment": "RDG need bool tensor support"},
 						{ "name": "NNXRuntimeHlsl", "skip": true, "comment": "RDG need bool tensor support"}						
 					]},
-				{	"target": "Div"},
+				{	"target": "Div", "tags": [ "SmokeTest" ]},
 				{	"target": "Equal", "output_type": "boolean", "runtimes": [
 						{ "name": "NNXRuntimeDml", "skip": true, "comment": "RDG need bool tensor support"},
 						{ "name": "NNXRuntimeHlsl", "skip": true, "comment": "RDG need bool tensor support"}						
 					]},
-				{	"target": "Greater", "output_type": "boolean", "runtimes": [
+				{	"target": "Greater", "tags": [ "SmokeTest" ], "output_type": "boolean", "runtimes": [
 						{ "name": "NNXRuntimeDml", "skip": true, "comment": "RDG need bool tensor support"},
 						{ "name": "NNXRuntimeHlsl", "skip": true, "comment": "RDG need bool tensor support"}						
 					]},
@@ -195,7 +194,7 @@
 						{ "name": "NNXRuntimeDml", "skip": true, "comment": "RDG need bool tensor support"},
 						{ "name": "NNXRuntimeHlsl", "skip": true, "comment": "RDG need bool tensor support"}						
 					]},
-				{	"target": "PRelu", "removed_datasets": [ 
+				{	"target": "PRelu", "tags": [ "SmokeTest" ], "removed_datasets": [ 
 						"Elementwise.Binary.BidirectionalBroadcast", 
 						"Elementwise.Binary.LHSToRHSBroadcast" ], 
 					"comment": "Prelu does not support bidirectional broadcast. It only support slope(RHS) to input0(LHS)"},
@@ -213,9 +212,9 @@
 			"runtimes": [ { "name": "NNXRuntimeDml", "skip": true} ],
 			"targets": [
 				{	"target": "Max", "additional_datasets": [ "Elementwise.Slow.TwoInputs" ]},
-				{	"target": "Mean"},
+				{	"target": "Mean", "tags": [ "SmokeTest" ]},
 				{	"target": "Min"},
-				{	"target": "Sum"}
+				{	"target": "Sum", "tags": [ "SmokeTest" ]}
 			]
 		}
 	],
