NNX Uplugin:
Fix Static Analysis by duplicating the utils file.
https://horde.devtools.epicgames.com/log/6335f57d7c11614c217c5f2e?issue=218727

diff --git a/Engine/Plugins/Experimental/NNX/Source/NNXRuntimeCPU/NNXRuntimeCPU.Build.cs b/Engine/Plugins/Experimental/NNX/Source/NNXRuntimeCPU/NNXRuntimeCPU.Build.cs
--- a/Engine/Plugins/Experimental/NNX/Source/NNXRuntimeCPU/NNXRuntimeCPU.Build.cs
+++ b/Engine/Plugins/Experimental/NNX/Source/NNXRuntimeCPU/NNXRuntimeCPU.Build.cs
@@ -13,12 +13,6 @@ public class NNXRuntimeCPU : ModuleRules
 		ShortName = "NNXRtCpu"; // Shorten to avoid path-too-long errors
 		PCHUsage = ModuleRules.PCHUsageMode.UseExplicitOrSharedPCHs;
 
-		PublicIncludePaths.AddRange(
-			new string[] {
-				System.IO.Path.Combine(ModuleDirectory, "../NNXRuntimeORT/Public/")
-			}
-		);
-
 		PublicDependencyModuleNames.AddRange
 			(
 			new string[] {
diff --git a/Engine/Plugins/Experimental/NNX/Source/NNXRuntimeCPU/Private/NNXRuntimeCPU.cpp b/Engine/Plugins/Experimental/NNX/Source/NNXRuntimeCPU/Private/NNXRuntimeCPU.cpp
--- a/Engine/Plugins/Experimental/NNX/Source/NNXRuntimeCPU/Private/NNXRuntimeCPU.cpp
+++ b/Engine/Plugins/Experimental/NNX/Source/NNXRuntimeCPU/Private/NNXRuntimeCPU.cpp
@@ -1,6 +1,6 @@
 // Copyright Epic Games, Inc. All Rights Reserved.
 #include "NNXRuntimeCPU.h"
-#include "NNXRuntimeORTUtils.h"
+#include "NNXRuntimeCPUUtils.h"
 #include "NeuralTimer.h"
 #include "RedirectCoutAndCerrToUeLog.h"
 
diff --git a/Engine/Plugins/Experimental/NNX/Source/NNXRuntimeCPU/Public/NNXRuntimeCPUUtils.h b/Engine/Plugins/Experimental/NNX/Source/NNXRuntimeCPU/Public/NNXRuntimeCPUUtils.h
new file mode 100644
--- /dev/null
+++ b/Engine/Plugins/Experimental/NNX/Source/NNXRuntimeCPU/Public/NNXRuntimeCPUUtils.h
@@ -0,0 +1,157 @@
+// Copyright Epic Games, Inc. All Rights Reserved.
+
+#pragma once
+
+#include "NNXTypes.h"
+#include "NNXRuntime.h"
+#include "NNXThirdPartyWarningDisabler.h"
+NNX_THIRD_PARTY_INCLUDES_START
+#undef check
+#undef TEXT
+#include "core/session/onnxruntime_cxx_api.h"
+NNX_THIRD_PARTY_INCLUDES_END
+
+DECLARE_STATS_GROUP(TEXT("MachineLearning"), STATGROUP_MachineLearning, STATCAT_Advanced);
+
+using TypeInfoORT = std::pair<EMLTensorDataType, uint64>;
+
+inline TypeInfoORT TranslateTensorTypeORTToNNI(unsigned int OrtDataType) 
+{
+	EMLTensorDataType DataType = EMLTensorDataType::None;
+	uint64 ElementSize = 0;
+
+	switch (OrtDataType) {
+	case ONNX_TENSOR_ELEMENT_DATA_TYPE_UNDEFINED:
+		DataType = EMLTensorDataType::None;
+		ElementSize = 0;
+		break;
+
+	case ONNX_TENSOR_ELEMENT_DATA_TYPE_FLOAT:
+		DataType = EMLTensorDataType::Float;
+		ElementSize = sizeof(float);
+		break;
+
+	case ONNX_TENSOR_ELEMENT_DATA_TYPE_UINT8:
+		DataType = EMLTensorDataType::UInt8;
+		ElementSize = sizeof(uint8);
+		break;
+	case ONNX_TENSOR_ELEMENT_DATA_TYPE_INT8:
+		DataType = EMLTensorDataType::Int8;
+		ElementSize = sizeof(int8);
+		break;
+	case ONNX_TENSOR_ELEMENT_DATA_TYPE_UINT16:
+		DataType = EMLTensorDataType::UInt16;
+		ElementSize = sizeof(uint16);
+		break;
+	case ONNX_TENSOR_ELEMENT_DATA_TYPE_INT16:
+		DataType = EMLTensorDataType::Int16;
+		ElementSize = sizeof(int16);
+		break;
+
+	case ONNX_TENSOR_ELEMENT_DATA_TYPE_INT32:
+		DataType = EMLTensorDataType::Int32;
+		ElementSize = sizeof(int32);
+		break;
+
+	case ONNX_TENSOR_ELEMENT_DATA_TYPE_INT64:
+		DataType = EMLTensorDataType::Int64;
+		ElementSize = sizeof(int64);
+		break;
+
+	case ONNX_TENSOR_ELEMENT_DATA_TYPE_STRING:
+		DataType = EMLTensorDataType::Char;
+		ElementSize = sizeof(char);
+		break;
+
+	case ONNX_TENSOR_ELEMENT_DATA_TYPE_BOOL:
+		DataType = EMLTensorDataType::Boolean;
+		ElementSize = sizeof(bool);
+		break;
+
+	case ONNX_TENSOR_ELEMENT_DATA_TYPE_FLOAT16:
+		DataType = EMLTensorDataType::Half;
+		ElementSize = 2;
+		break;
+
+	case ONNX_TENSOR_ELEMENT_DATA_TYPE_DOUBLE:
+		DataType = EMLTensorDataType::Double;
+		ElementSize = sizeof(double);
+		break;
+
+	case ONNX_TENSOR_ELEMENT_DATA_TYPE_UINT32:
+		DataType = EMLTensorDataType::UInt32;
+		ElementSize = sizeof(uint32);
+		break;
+
+	case ONNX_TENSOR_ELEMENT_DATA_TYPE_UINT64:
+		DataType = EMLTensorDataType::UInt64;
+		ElementSize = sizeof(uint64);
+		break;
+
+	case ONNX_TENSOR_ELEMENT_DATA_TYPE_COMPLEX64:
+		DataType = EMLTensorDataType::Complex64;
+		ElementSize = 8;
+		break;
+
+	case ONNX_TENSOR_ELEMENT_DATA_TYPE_COMPLEX128:
+		DataType = EMLTensorDataType::Complex128;
+		ElementSize = 16;
+		break;
+
+	case ONNX_TENSOR_ELEMENT_DATA_TYPE_BFLOAT16:
+		DataType = EMLTensorDataType::BFloat16;
+		ElementSize = 2;
+		break;
+
+	default:
+		DataType = EMLTensorDataType::None;
+		break;
+	}
+
+	return TypeInfoORT{DataType, ElementSize};
+}
+
+inline void BindTensorsToORT(
+	TArrayView<const NNX::FMLTensorBinding> InBindingTensors,
+	const TArray<NNX::FMLTensorDesc>& InTensorsDescriptors,
+	const TArray<ONNXTensorElementDataType>& InTensorsORTType, 
+	const Ort::MemoryInfo* InAllocatorInfo,
+	TArray<Ort::Value>& OutOrtTensors	
+)
+{
+	const uint32 NumBinding = InBindingTensors.Num();
+	const uint32 NumDescriptors = InTensorsDescriptors.Num();
+
+	if (NumBinding != NumDescriptors)
+	{
+		UE_LOG(LogNNX, Error, TEXT("BindTensorsToORT: Number of Bindings is different of Descriptors."));
+		return;
+	}
+
+	for (uint32 Index = 0; Index < NumBinding; ++Index)
+	{
+		const NNX::FMLTensorBinding& CurrentBinding = InBindingTensors[Index];
+		const NNX::FMLTensorDesc& CurrentDescriptor = InTensorsDescriptors[Index];
+		ONNXTensorElementDataType CurrentORTType = InTensorsORTType[Index];
+
+		TUniquePtr<int64_t[]> SizesInt64t;
+		SizesInt64t = MakeUnique<int64_t[]>(CurrentDescriptor.Dimension);
+		for (uint32 DimIndex = 0; DimIndex < CurrentDescriptor.Dimension; ++DimIndex)
+		{
+			SizesInt64t.Get()[DimIndex] = CurrentDescriptor.Sizes[DimIndex];
+		}
+
+		const uint64 ByteCount { InTensorsDescriptors[Index].DataSize };
+		const uint32 ArrayDimensions { CurrentDescriptor.Dimension };
+		OutOrtTensors.Emplace(
+			Ort::Value::CreateTensor(
+				*InAllocatorInfo,
+				CurrentBinding.CpuMemory,
+				ByteCount, 
+				SizesInt64t.Get(),
+				ArrayDimensions, 
+				CurrentORTType
+			)
+		);
+	}
+}
\ No newline at end of file
diff --git a/Engine/Plugins/Experimental/NNX/Source/NNXRuntimeORT/Private/NNXRuntimeORTUtils.h b/Engine/Plugins/Experimental/NNX/Source/NNXRuntimeORT/Private/NNXRuntimeORTUtils.h
new file mode 100644
--- /dev/null
+++ b/Engine/Plugins/Experimental/NNX/Source/NNXRuntimeORT/Private/NNXRuntimeORTUtils.h
@@ -0,0 +1,157 @@
+// Copyright Epic Games, Inc. All Rights Reserved.
+
+#pragma once
+
+#include "NNXTypes.h"
+#include "NNXRuntime.h"
+#include "NNXThirdPartyWarningDisabler.h"
+NNX_THIRD_PARTY_INCLUDES_START
+#undef check
+#undef TEXT
+#include "core/session/onnxruntime_cxx_api.h"
+NNX_THIRD_PARTY_INCLUDES_END
+
+DECLARE_STATS_GROUP(TEXT("MachineLearning"), STATGROUP_MachineLearning, STATCAT_Advanced);
+
+using TypeInfoORT = std::pair<EMLTensorDataType, uint64>;
+
+inline TypeInfoORT TranslateTensorTypeORTToNNI(unsigned int OrtDataType) 
+{
+	EMLTensorDataType DataType = EMLTensorDataType::None;
+	uint64 ElementSize = 0;
+
+	switch (OrtDataType) {
+	case ONNX_TENSOR_ELEMENT_DATA_TYPE_UNDEFINED:
+		DataType = EMLTensorDataType::None;
+		ElementSize = 0;
+		break;
+
+	case ONNX_TENSOR_ELEMENT_DATA_TYPE_FLOAT:
+		DataType = EMLTensorDataType::Float;
+		ElementSize = sizeof(float);
+		break;
+
+	case ONNX_TENSOR_ELEMENT_DATA_TYPE_UINT8:
+		DataType = EMLTensorDataType::UInt8;
+		ElementSize = sizeof(uint8);
+		break;
+	case ONNX_TENSOR_ELEMENT_DATA_TYPE_INT8:
+		DataType = EMLTensorDataType::Int8;
+		ElementSize = sizeof(int8);
+		break;
+	case ONNX_TENSOR_ELEMENT_DATA_TYPE_UINT16:
+		DataType = EMLTensorDataType::UInt16;
+		ElementSize = sizeof(uint16);
+		break;
+	case ONNX_TENSOR_ELEMENT_DATA_TYPE_INT16:
+		DataType = EMLTensorDataType::Int16;
+		ElementSize = sizeof(int16);
+		break;
+
+	case ONNX_TENSOR_ELEMENT_DATA_TYPE_INT32:
+		DataType = EMLTensorDataType::Int32;
+		ElementSize = sizeof(int32);
+		break;
+
+	case ONNX_TENSOR_ELEMENT_DATA_TYPE_INT64:
+		DataType = EMLTensorDataType::Int64;
+		ElementSize = sizeof(int64);
+		break;
+
+	case ONNX_TENSOR_ELEMENT_DATA_TYPE_STRING:
+		DataType = EMLTensorDataType::Char;
+		ElementSize = sizeof(char);
+		break;
+
+	case ONNX_TENSOR_ELEMENT_DATA_TYPE_BOOL:
+		DataType = EMLTensorDataType::Boolean;
+		ElementSize = sizeof(bool);
+		break;
+
+	case ONNX_TENSOR_ELEMENT_DATA_TYPE_FLOAT16:
+		DataType = EMLTensorDataType::Half;
+		ElementSize = 2;
+		break;
+
+	case ONNX_TENSOR_ELEMENT_DATA_TYPE_DOUBLE:
+		DataType = EMLTensorDataType::Double;
+		ElementSize = sizeof(double);
+		break;
+
+	case ONNX_TENSOR_ELEMENT_DATA_TYPE_UINT32:
+		DataType = EMLTensorDataType::UInt32;
+		ElementSize = sizeof(uint32);
+		break;
+
+	case ONNX_TENSOR_ELEMENT_DATA_TYPE_UINT64:
+		DataType = EMLTensorDataType::UInt64;
+		ElementSize = sizeof(uint64);
+		break;
+
+	case ONNX_TENSOR_ELEMENT_DATA_TYPE_COMPLEX64:
+		DataType = EMLTensorDataType::Complex64;
+		ElementSize = 8;
+		break;
+
+	case ONNX_TENSOR_ELEMENT_DATA_TYPE_COMPLEX128:
+		DataType = EMLTensorDataType::Complex128;
+		ElementSize = 16;
+		break;
+
+	case ONNX_TENSOR_ELEMENT_DATA_TYPE_BFLOAT16:
+		DataType = EMLTensorDataType::BFloat16;
+		ElementSize = 2;
+		break;
+
+	default:
+		DataType = EMLTensorDataType::None;
+		break;
+	}
+
+	return TypeInfoORT{DataType, ElementSize};
+}
+
+inline void BindTensorsToORT(
+	TArrayView<const NNX::FMLTensorBinding> InBindingTensors,
+	const TArray<NNX::FMLTensorDesc>& InTensorsDescriptors,
+	const TArray<ONNXTensorElementDataType>& InTensorsORTType, 
+	const Ort::MemoryInfo* InAllocatorInfo,
+	TArray<Ort::Value>& OutOrtTensors	
+)
+{
+	const uint32 NumBinding = InBindingTensors.Num();
+	const uint32 NumDescriptors = InTensorsDescriptors.Num();
+
+	if (NumBinding != NumDescriptors)
+	{
+		UE_LOG(LogNNX, Error, TEXT("BindTensorsToORT: Number of Bindings is different of Descriptors."));
+		return;
+	}
+
+	for (uint32 Index = 0; Index < NumBinding; ++Index)
+	{
+		const NNX::FMLTensorBinding& CurrentBinding = InBindingTensors[Index];
+		const NNX::FMLTensorDesc& CurrentDescriptor = InTensorsDescriptors[Index];
+		ONNXTensorElementDataType CurrentORTType = InTensorsORTType[Index];
+
+		TUniquePtr<int64_t[]> SizesInt64t;
+		SizesInt64t = MakeUnique<int64_t[]>(CurrentDescriptor.Dimension);
+		for (uint32 DimIndex = 0; DimIndex < CurrentDescriptor.Dimension; ++DimIndex)
+		{
+			SizesInt64t.Get()[DimIndex] = CurrentDescriptor.Sizes[DimIndex];
+		}
+
+		const uint64 ByteCount { InTensorsDescriptors[Index].DataSize };
+		const uint32 ArrayDimensions { CurrentDescriptor.Dimension };
+		OutOrtTensors.Emplace(
+			Ort::Value::CreateTensor(
+				*InAllocatorInfo,
+				CurrentBinding.CpuMemory,
+				ByteCount, 
+				SizesInt64t.Get(),
+				ArrayDimensions, 
+				CurrentORTType
+			)
+		);
+	}
+}
\ No newline at end of file
diff --git a/Engine/Plugins/Experimental/NNX/Source/NNXRuntimeORT/Public/NNXRuntimeORTUtils.h b/Engine/Plugins/Experimental/NNX/Source/NNXRuntimeORT/Public/NNXRuntimeORTUtils.h
deleted file mode 100644
