NNXQA Fixes for running test on packaged exe and platforms

Problem: smoke test are not running on platform and packaged builds

Solution (partial more will need to be done):
* Add json test description as runtime ressource so they are accessible by packaged game
* Use NNXRuntimeCPU as reference as NNXRuntimeORTCpu is not availalble on all platforms
* Reduce the amount of tests on NNXIncubator start to Acos op on HLSL runtime (as we can manually trigger test via command line or session frontend).

Test:
* Verified json are copied at the right place on packaged builds
* Runned smoke test on standalone build on windows and on platform (both fail for different reasons, this CL partially fix it)

diff --git a/Engine/Plugins/Experimental/NNX/Source/NNXQA/NNXQA.Build.cs b/Engine/Plugins/Experimental/NNX/Source/NNXQA/NNXQA.Build.cs
--- a/Engine/Plugins/Experimental/NNX/Source/NNXQA/NNXQA.Build.cs
+++ b/Engine/Plugins/Experimental/NNX/Source/NNXQA/NNXQA.Build.cs
@@ -1,5 +1,6 @@
 // Copyright Epic Games, Inc. All Rights Reserved.
 
+using System.IO;
 using UnrealBuildTool;
 
 public class NNXQA : ModuleRules
@@ -25,5 +26,11 @@ public class NNXQA : ModuleRules
 				"NNXUtils"
 			}
 		);
+
+		// RuntimeDependencies
+		foreach (string JsonFilePath in Directory.EnumerateFiles(Path.Combine(ModuleDirectory, "resources"), "*.json"))
+		{
+			RuntimeDependencies.Add(JsonFilePath);
+		}
 	}
 }
diff --git a/Engine/Plugins/Experimental/NNX/Source/NNXQA/Private/NNXQAUtils.cpp b/Engine/Plugins/Experimental/NNX/Source/NNXQA/Private/NNXQAUtils.cpp
--- a/Engine/Plugins/Experimental/NNX/Source/NNXQA/Private/NNXQAUtils.cpp
+++ b/Engine/Plugins/Experimental/NNX/Source/NNXQA/Private/NNXQAUtils.cpp
@@ -374,10 +374,10 @@ namespace Test
 		UE_LOG(LogNNX, Display, TEXT("Starting tests of '%s'"), *TestSetup.TargetName);
 
 		// Reference runtime
-		NNX::IRuntime* RefRuntime = NNX::GetRuntime(TEXT("NNXRuntimeORTCpu"));
+		NNX::IRuntime* RefRuntime = NNX::GetRuntime(TEXT("NNXRuntimeCPU"));
 		if (!RefRuntime)
 		{
-			UE_LOG(LogNNX, Error, TEXT("Can't load NNXRuntimeORTCpu runtime. Tests ABORTED!"));
+			UE_LOG(LogNNX, Error, TEXT("Can't load NNXRuntimeCPU runtime. Tests ABORTED!"));
 			return false;
 		}
 		const FString& RefName = RefRuntime->GetRuntimeName();
