Fixed string conversion problem when adding operator to ONNX graph

#rnx

diff --git a/Engine/Plugins/Experimental/NNX/Source/NNXUtils/Private/NNXModelBuilderONNX.cpp b/Engine/Plugins/Experimental/NNX/Source/NNXUtils/Private/NNXModelBuilderONNX.cpp
--- a/Engine/Plugins/Experimental/NNX/Source/NNXUtils/Private/NNXModelBuilderONNX.cpp
+++ b/Engine/Plugins/Experimental/NNX/Source/NNXUtils/Private/NNXModelBuilderONNX.cpp
@@ -153,9 +153,11 @@ public:
 	/** Add operator */
 	virtual HOperator AddOperator(const FString& TypeName, const FString& Name) override
 	{
-		auto Node = Graph->add_node();
-		auto NodeTypeNameStr = TCHAR_TO_ANSI(*TypeName);
-		auto NodeType = onnx::Symbol(NodeTypeNameStr).toString();
+		FTCHARToUTF8 Convert(*TypeName);
+
+		std::string NodeType = onnx::Symbol(Convert.Get()).toString();
+
+		onnx::NodeProto* Node = Graph->add_node();
 
 		Node->set_op_type(NodeType);
 
@@ -167,6 +169,7 @@ public:
 		{
 			Node->set_name(TCHAR_TO_ANSI(*TypeName));
 		}
+
 		Node->set_domain(Model.domain());
 
 		return MakeOperatorHandle(Node);
