NNXQA fix build script for linux

#rnx

diff --git a/Engine/Plugins/Experimental/NNX/Source/NNXQA/NNXQA.Build.cs b/Engine/Plugins/Experimental/NNX/Source/NNXQA/NNXQA.Build.cs
--- a/Engine/Plugins/Experimental/NNX/Source/NNXQA/NNXQA.Build.cs
+++ b/Engine/Plugins/Experimental/NNX/Source/NNXQA/NNXQA.Build.cs
@@ -28,7 +28,7 @@ public class NNXQA : ModuleRules
 		);
 
 		// RuntimeDependencies
-		foreach (string JsonFilePath in Directory.EnumerateFiles(Path.Combine(ModuleDirectory, "resources"), "*.json"))
+		foreach (string JsonFilePath in Directory.EnumerateFiles(Path.Combine(ModuleDirectory, "Resources"), "*.json"))
 		{
 			RuntimeDependencies.Add(JsonFilePath);
 		}
