NNX only allow ORT shared lib based module on editor

#rnx

diff --git a/Engine/Plugins/Experimental/NNX/NNX.uplugin b/Engine/Plugins/Experimental/NNX/NNX.uplugin
--- a/Engine/Plugins/Experimental/NNX/NNX.uplugin
+++ b/Engine/Plugins/Experimental/NNX/NNX.uplugin
@@ -51,7 +51,7 @@
 		},
 		{
 			"Name": "NNXRuntimeCPU",
-			"Type": "Runtime",
+			"Type": "RuntimeAndProgram",
 			"LoadingPhase": "Default",
 			"PlatformAllowList": [ "Win64", "Linux" ]
 		},
@@ -63,13 +63,13 @@
 		},
 		{
 			"Name": "NNXRuntimeORT",
-			"Type": "Runtime",
+			"Type": "EditorAndProgram",
 			"LoadingPhase": "Default",
 			"PlatformAllowList": [ "Win64", "Linux" ]
 		},
 		{
 			"Name": "NNX_ONNXRuntime",
-			"Type": "Runtime",
+			"Type": "RuntimeAndProgram",
 			"LoadingPhase": "Default",
 			"PlatformAllowList": [ "Win64", "Linux" ]
 		},
