- Added check for current pipeline
- Refactored the readback code

#rnx

diff --git a/Engine/Plugins/Experimental/NNX/Source/NNXRuntimeRDG/Private/NNXRuntimeRDG.cpp b/Engine/Plugins/Experimental/NNX/Source/NNXRuntimeRDG/Private/NNXRuntimeRDG.cpp
--- a/Engine/Plugins/Experimental/NNX/Source/NNXRuntimeRDG/Private/NNXRuntimeRDG.cpp
+++ b/Engine/Plugins/Experimental/NNX/Source/NNXRuntimeRDG/Private/NNXRuntimeRDG.cpp
@@ -89,6 +89,13 @@ int FMLInferenceModelRDG::Run(TArrayView<const FMLTensorBinding> InInputBindings
 	(
 		[&Signal, &Res, this, InInputBindings, OutOutputBindings](FRHICommandListImmediate& RHICmdList)
 		{
+			TOptional<ERHIPipeline>		Pipeline = RHICmdList.GetPipeline();
+
+			if (Pipeline == ERHIPipeline::None)
+			{
+				RHICmdList.SwitchPipeline(ERHIPipeline::Graphics);
+			}
+
 			FRDGBuilder	GraphBuilder(RHICmdList);
 
 			Res = EnqueueRDG(GraphBuilder, InInputBindings, OutOutputBindings);
@@ -243,7 +250,7 @@ void FMLInferenceModelRDG::AddTensorReadbacks_RenderThread(FRDGBuilder& GraphBui
 		GraphBuilder.AddPass(
 			RDG_EVENT_NAME("FMLInferenceModelAddTensorReadback"),
 			TensorReadbackParams,
-			ERDGPassFlags::Copy | ERDGPassFlags::NeverCull,
+			ERDGPassFlags::Readback | ERDGPassFlags::NeverCull,
 			[InBinding, TensorDesc, TensorReadbackParams](FRHICommandListImmediate& RHICmdList)
 			{
 				FRHIBuffer* OutputBuffer = TensorReadbackParams->Buffer->GetRHI();
@@ -255,12 +262,22 @@ void FMLInferenceModelRDG::AddTensorReadbacks_RenderThread(FRDGBuilder& GraphBui
 
 				RHICmdList.Transition(MakeArrayView(PreTransitions, UE_ARRAY_COUNT(PreTransitions)));
 
-				// TODO: FIXME: We need this for DirectML
+				// TODO: FIXME: We need to transition the resources for DirectML
 				RHICmdList.SubmitCommandsHint();
 
-				const void* BuffData = RHICmdList.LockBuffer(OutputBuffer, 0, TensorDesc.DataSize, RLM_ReadOnly);
+				FRHIGPUBufferReadback	Readback("FMLTensorReadback");
+
+				Readback.EnqueueCopy(RHICmdList, OutputBuffer, TensorDesc.DataSize);
+				RHICmdList.BlockUntilGPUIdle();
+				check(Readback.IsReady());
+
+				const void* BuffData = Readback.Lock(TensorDesc.DataSize);
 				FMemory::Memcpy(InBinding.CpuMemory, BuffData, TensorDesc.DataSize);
-				RHICmdList.UnlockBuffer(OutputBuffer);
+				Readback.Unlock();
+
+				//const void* BuffData = RHICmdList.LockBuffer(OutputBuffer, 0, TensorDesc.DataSize, RLM_ReadOnly);
+				//FMemory::Memcpy(InBinding.CpuMemory, BuffData, TensorDesc.DataSize);
+				//RHICmdList.UnlockBuffer(OutputBuffer);
 
 				//FRHITransitionInfo PostTransitions[] =
 				//{
