Operator registration and tests for HLSL runtime GEMM operator.

#rnx

diff --git a/Engine/Plugins/Experimental/NNX/Source/NNXHlslShaders/Private/NNXGemmCS.cpp b/Engine/Plugins/Experimental/NNX/Source/NNXHlslShaders/Private/NNXGemmCS.cpp
--- a/Engine/Plugins/Experimental/NNX/Source/NNXHlslShaders/Private/NNXGemmCS.cpp
+++ b/Engine/Plugins/Experimental/NNX/Source/NNXHlslShaders/Private/NNXGemmCS.cpp
@@ -15,9 +15,11 @@ namespace
 		TArray<uint32> StackShapeB;
 		TArray<uint32> StackStrideA;
 		TArray<uint32> StackStrideB;
+		uint32 CWidth = 0;
+		uint32 CHeight = 0;
 	};
 
-	static GemmMatrixParameters GetMatrixParameters(TArray<uint32> ShapeA, TArray<uint32> ShapeB)
+	static GemmMatrixParameters GetMatrixParameters(TArray<uint32> ShapeA, TArray<uint32> ShapeB, TArray<uint32> ShapeC)
 	{
 		check(!ShapeA.IsEmpty());
 		check(!ShapeB.IsEmpty());
@@ -44,6 +46,7 @@ namespace
 
 		const bool IsVectorA = ShapeA.Num() == 1;
 		const bool IsVectorB = ShapeB.Num() == 1;
+		const bool IsVectorC = ShapeC.Num() == 1;
 
 		Result.M = IsVectorA ? 1 : ShapeA[ShapeA.Num() - 2];
 		Result.N = IsVectorB ? 1 : ShapeB[ShapeB.Num() - 1];
@@ -68,6 +71,9 @@ namespace
 			Result.StackStrideB[i] = Result.StackStrideB[i + 1] * Result.StackShapeB[i + 1];
 		}
 
+		Result.CWidth = ShapeC.IsEmpty() ? 0 : ShapeC[ShapeC.Num() - 1];
+		Result.CHeight = ShapeC.IsEmpty() ? 0 : IsVectorC ? 1 : ShapeC[ShapeC.Num() - 2];
+
 		return Result;
 	}
 } // namespace
@@ -78,20 +84,31 @@ void FMLGemmCS::ModifyCompilationEnvironment(const FGlobalShaderPermutationParam
 	OutEnvironment.SetDefine(TEXT("MAX_NUM_STACK_DIMENSIONS"), NNXRT_GEMM_MAX_NUM_STACK_DIMENSIONS);
 }
 
-void FMLGemmCS::FillInParameters(float Alpha, float Beta, int32 TransA, int32 TransB, uint32 M, uint32 N, uint32 K, uint32 CWidth, uint32 CHeight, float CScalar, FMLGemmCS::FParameters& Parameters)
+void FMLGemmCS::FillInParameters(float Alpha, float Beta, int32 TransA, int32 TransB, const NNX::FMLTensorDesc &InputA, const NNX::FMLTensorDesc &InputB,
+		const NNX::FMLTensorDesc &InputC, float CScalar, FMLGemmCS::FParameters& Parameters)
 {
+	auto MakeShape = [](const NNX::FMLTensorDesc &InputDesc) {
+		TArray<uint32> Result;
+		for (int32 i = 0; i < (int32)InputDesc.Dimension; i++) {
+			Result.Add(InputDesc.Sizes[i]);
+		}
+		return Result;
+	};
+
+	GemmMatrixParameters Params = GetMatrixParameters(MakeShape(InputA), MakeShape(InputB), MakeShape(InputC));
+
 	Parameters.Alpha = Alpha;
 	Parameters.Beta = Beta;
 	Parameters.TransA = TransA;
 	Parameters.TransB = TransB;
-	Parameters.M = M;
-	Parameters.N = N;
-	Parameters.K = K;
-	Parameters.MxK = M * K;
-	Parameters.KxN = K * N;
-	Parameters.MxN = M * N;
-	Parameters.CWidth = CWidth;
-	Parameters.CHeight = CHeight;
+	Parameters.M = Params.M;
+	Parameters.N = Params.N;
+	Parameters.K = Params.K;
+	Parameters.MxK = Params.M * Params.K;
+	Parameters.KxN = Params.K * Params.N;
+	Parameters.MxN = Params.M * Params.N;
+	Parameters.CWidth = Params.CWidth;
+	Parameters.CHeight = Params.CHeight;
 	Parameters.CScalar = CScalar;
 }
 
@@ -104,14 +121,14 @@ uint32 FMLGemmCS::GetShapeSize(TArray<uint32> Shape)
 
 uint32 FMLGemmCS::GetMatMulOutputSize(TArray<uint32> ShapeA, TArray<uint32> ShapeB)
 {
-	GemmMatrixParameters Params = GetMatrixParameters(ShapeA, ShapeB);
+	GemmMatrixParameters Params = GetMatrixParameters(ShapeA, ShapeB, {});
 
 	return FMath::Max(GetShapeSize(Params.StackShapeA), GetShapeSize(Params.StackShapeB)) * Params.M * Params.N;
 }
 
 void FMLGemmCS::FillInParametersMatMul(TArray<uint32> ShapeA, TArray<uint32> ShapeB, FMLGemmCS::FParameters& Parameters)
 {
-	GemmMatrixParameters Params = GetMatrixParameters(ShapeA, ShapeB);
+	GemmMatrixParameters Params = GetMatrixParameters(ShapeA, ShapeB, {});
 
 	Parameters.Alpha = 1.0f;
 	Parameters.Beta = 1.0f;
@@ -123,9 +140,9 @@ void FMLGemmCS::FillInParametersMatMul(TArray<uint32> ShapeA, TArray<uint32> Sha
 	Parameters.MxK = Params.M * Params.K;
 	Parameters.KxN = Params.K * Params.N;
 	Parameters.MxN = Params.M * Params.N;
-	Parameters.CWidth = 0;
-	Parameters.CHeight = 0;
-	Parameters.CScalar = 0;
+	Parameters.CWidth = Params.CWidth;
+	Parameters.CHeight = Params.CHeight;
+	Parameters.CScalar = 0.0f;
 
 	for (int32 i = 0; i < Params.StackShapeA.Num(); i++)
 	{
diff --git a/Engine/Plugins/Experimental/NNX/Source/NNXHlslShaders/Public/NNXGemmCS.h b/Engine/Plugins/Experimental/NNX/Source/NNXHlslShaders/Public/NNXGemmCS.h
--- a/Engine/Plugins/Experimental/NNX/Source/NNXHlslShaders/Public/NNXGemmCS.h
+++ b/Engine/Plugins/Experimental/NNX/Source/NNXHlslShaders/Public/NNXGemmCS.h
@@ -4,6 +4,7 @@
 
 #include "GlobalShader.h"
 #include "NNXShaderParameters.h"
+#include "NNXTypes.h"
 
 enum class EGemmCScalar : uint8
 {
@@ -48,7 +49,8 @@ public:
 
 	static void ModifyCompilationEnvironment(const FGlobalShaderPermutationParameters& InParameters, FShaderCompilerEnvironment& OutEnvironment);
 
-	static void FillInParameters(float Alpha, float Beta, int32 TransA, int32 TransB, uint32 M, uint32 N, uint32 K, uint32 CWidth, uint32 CHeight, float CScalar, FMLGemmCS::FParameters& Parameters);
+	static void FillInParameters(float Alpha, float Beta, int32 TransA, int32 TransB, const NNX::FMLTensorDesc &InputA, const NNX::FMLTensorDesc &InputB,
+		const NNX::FMLTensorDesc &InputC, float CScalar, FMLGemmCS::FParameters& Parameters);
 
 	static uint32 GetShapeSize(TArray<uint32> Shape);
 	static uint32 GetMatMulOutputSize(TArray<uint32> ShapeA, TArray<uint32> ShapeB);
diff --git a/Engine/Plugins/Experimental/NNX/Source/NNXQA/Resources/NNXQATestDesc.json b/Engine/Plugins/Experimental/NNX/Source/NNXQA/Resources/NNXQATestDesc.json
--- a/Engine/Plugins/Experimental/NNX/Source/NNXQA/Resources/NNXQATestDesc.json
+++ b/Engine/Plugins/Experimental/NNX/Source/NNXQA/Resources/NNXQATestDesc.json
@@ -216,6 +216,13 @@
 				{	"target": "Min"},
 				{	"target": "Sum", "tags": [ "SmokeTest" ]}
 			]
+		},
+		{
+			"category": "Gemm",
+			"runtimes": [ { "name": "NNXRuntimeDml", "skip": true } ],
+			"targets": [
+				{ "target": "Gemm" }
+			]
 		}
 	],
 
@@ -339,6 +346,15 @@
 					"inputs": [ { "shape": [ 65536, 513 ] }, { "shape": [ 1 ] } ]
 				}
 			]
+		},
+		{
+			"name": "Gemm",
+			"datasets": [
+				{ "inputs": [ { "shape": [ 2, 2 ] }, { "shape": [ 2, 2 ] } ], "outputs": [ { "shape": [ 2, 2 ] } ]},
+				{ "inputs": [ { "shape": [ 3, 4 ] }, { "shape": [ 4, 5 ] } ], "outputs": [ { "shape": [ 3, 5 ] } ]},
+				{ "inputs": [ { "shape": [ 3, 4 ] }, { "shape": [ 4, 5 ] }, { "shape": [ 3, 5 ] } ], "outputs": [ { "shape": [ 3, 5 ] } ]},
+				{ "inputs": [ { "shape": [ 1, 1 ] }, { "shape": [ 1, 5 ] }, { "shape": [ 5 ] } ], "outputs": [ { "shape": [ 1, 5 ] } ]}
+			]
 		}
 	]
 	
diff --git a/Engine/Plugins/Experimental/NNX/Source/NNXRuntimeRDG/Private/NNXRuntimeHLSL.cpp b/Engine/Plugins/Experimental/NNX/Source/NNXRuntimeRDG/Private/NNXRuntimeHLSL.cpp
--- a/Engine/Plugins/Experimental/NNX/Source/NNXRuntimeRDG/Private/NNXRuntimeHLSL.cpp
+++ b/Engine/Plugins/Experimental/NNX/Source/NNXRuntimeRDG/Private/NNXRuntimeHLSL.cpp
@@ -6,6 +6,7 @@
 #include "NNXRuntimeHLSLElementWiseBinaryOps.h"
 #include "NNXRuntimeHLSLElementWiseUnaryOps.h"
 #include "NNXRuntimeHLSLElementWiseVariadicOps.h"
+#include "NNXRuntimeHLSLGemmOp.h"
 #include "NNXRuntimeHLSLHelper.h"
 #include "NNXRuntimeRDG.h"
 
@@ -151,6 +152,7 @@ public:
 		RegisterElementWiseUnaryOperators(*registry);
 		RegisterElementWiseBinaryOperators(*registry);
 		RegisterElementWiseVariadicOperators(*registry);
+		RegisterGemmOperator(*registry);
 
 		return true;
 	}
diff --git a/Engine/Plugins/Experimental/NNX/Source/NNXRuntimeRDG/Private/NNXRuntimeHLSLGemmOp.cpp b/Engine/Plugins/Experimental/NNX/Source/NNXRuntimeRDG/Private/NNXRuntimeHLSLGemmOp.cpp
new file mode 100644
--- /dev/null
+++ b/Engine/Plugins/Experimental/NNX/Source/NNXRuntimeRDG/Private/NNXRuntimeHLSLGemmOp.cpp
@@ -0,0 +1,113 @@
+// Copyright Epic Games, Inc. All Rights Reserved.
+
+#include "NNXRuntimeHLSLGemmOp.h"
+#include "NNXGemmCS.h"
+#include "NNXRuntimeHLSLHelper.h"
+
+namespace NNX
+{
+
+	/**
+	 * Gemm ML operator implementation
+	 */
+	class FMLOperatorHlslGemm : public FMLOperatorHlsl
+	{
+	public:
+
+		static FMLOperatorHlsl* Create()
+		{
+			return new FMLOperatorHlslGemm();
+		}
+
+		virtual ~FMLOperatorHlslGemm() = default;
+
+	private:
+
+		FMLOperatorHlslGemm() {}
+
+		FMLTensorDesc InputA = {};
+		FMLTensorDesc InputB = {};
+		FMLTensorDesc InputC = {};
+		FMLTensorDesc Output = {};
+
+		float InputAlpha = 1.0f;
+		float InputBeta = 1.0f;
+		int32 InputTransA = 0;
+		int32 InputTransB = 0;
+		uint32 InputM = 0;
+		uint32 InputN = 0;
+		uint32 InputK = 0;
+		uint32 InputCWidth = 0;
+		uint32 InputCHeight = 0;
+
+		bool bIsCScalar = false;
+		bool bNoBias = true;
+
+	public:
+
+		virtual bool Initialize(TArrayView<const FMLTensorDesc> InputTensors, TArrayView<const FMLTensorDesc> OutputTensors) override
+		{
+			check(InputTensors.Num() >= 2 && InputTensors.Num() <= 3);
+			check(OutputTensors.Num() == 1);
+
+			InputA = InputTensors[0];
+			InputB = InputTensors[1];
+			if (InputTensors.Num() == 3) InputC = InputTensors[2];
+			Output = OutputTensors[0];
+
+			check(InputA.Dimension == 2);
+			check(InputB.Dimension == 2);
+			check(InputC.Dimension < 3);
+
+			check(InputC.Dimension != 1 || InputC.Sizes[0] != 1); // TODO scalar version not supported yet
+
+			// C is treated as a scalar if there is no valid C, either width or height is zero or C dimension is 1x1
+			bIsCScalar = false; // InputTensors.Num() != 3 || InputC.Sizes[0] * InputC.Sizes[1] < 2;
+			// CScalar = C != nullptr ? C[0] : (InElementType)0;
+			bNoBias = InputTensors.Num() != 3 /*|| InputC.Sizes[0] * InputC.Sizes[1] < 1*/;
+
+			return true;
+		}
+
+		virtual void Dispatch(FRDGBuilder& GraphBuilder, TArrayView<const FMLTensorBinding> InInputBindings, TArrayView<const FMLTensorBinding> OutOutputBindings) override
+		{
+			const EGemmAlgorithm Algorithm = EGemmAlgorithm::Simple32x32;
+
+			const float CScalar = 0.0f;
+
+			// Set parameters
+			FMLGemmCS::FParameters* Parameters = GraphBuilder.AllocParameters<FMLGemmCS::FParameters>();
+			FMLGemmCS::FillInParameters(InputAlpha, InputBeta, InputTransA, InputTransB, InputA, InputB, InputC, CScalar, *Parameters);
+			Parameters->A = GraphBuilder.CreateSRV(FRDGBufferSRVDesc(InInputBindings[0].Buffer, PF_R32_FLOAT));
+			Parameters->B = GraphBuilder.CreateSRV(FRDGBufferSRVDesc(InInputBindings[1].Buffer, PF_R32_FLOAT));
+			if (InInputBindings.Num() == 3) {
+				Parameters->C = GraphBuilder.CreateSRV(FRDGBufferSRVDesc(InInputBindings[2].Buffer, PF_R32_FLOAT));
+			}
+			Parameters->Y = GraphBuilder.CreateUAV(FRDGBufferUAVDesc(OutOutputBindings[0].Buffer, PF_R32_FLOAT));
+
+			FMLGemmCS::FPermutationDomain PermutationVector;
+			PermutationVector.Set<FMLGemmCS::FGemmCScalar>(bNoBias ? EGemmCScalar::NoBias : (bIsCScalar ? EGemmCScalar::Yes : EGemmCScalar::No));
+			PermutationVector.Set<FMLGemmCS::FGemmAlgorithm>(Algorithm);
+			PermutationVector.Set<FMLGemmCS::FGemmNumStackDimensions>(0);
+			TShaderMapRef<FMLGemmCS> ComputeShader(GetGlobalShaderMap(GMaxRHIFeatureLevel), PermutationVector);
+
+			FIntVector ThreadGroupCount = FMLGemmCS::GetGroupCount(*Parameters, Algorithm, 0);
+
+			FComputeShaderUtils::AddPass(
+				GraphBuilder,
+				RDG_EVENT_NAME("FMLHlslGemmOperatorHlsl_Dispatch"),
+				ERDGPassFlags::Compute | ERDGPassFlags::NeverCull,
+				ComputeShader,
+				Parameters,
+				ThreadGroupCount);
+		}
+	};
+
+	bool RegisterGemmOperator(FMLOperatorRegistryHlsl& Registry)
+	{
+		Registry.OpAdd(TEXT("Gemm"), FMLOperatorHlslGemm::Create);
+
+		return true;
+	}
+
+} // NNX
diff --git a/Engine/Plugins/Experimental/NNX/Source/NNXRuntimeRDG/Private/NNXRuntimeHLSLGemmOp.h b/Engine/Plugins/Experimental/NNX/Source/NNXRuntimeRDG/Private/NNXRuntimeHLSLGemmOp.h
new file mode 100644
--- /dev/null
+++ b/Engine/Plugins/Experimental/NNX/Source/NNXRuntimeRDG/Private/NNXRuntimeHLSLGemmOp.h
@@ -0,0 +1,11 @@
+// Copyright Epic Games, Inc. All Rights Reserved.
+
+#pragma once
+
+#include "NNXRuntimeRDG.h"
+#include "NNXRuntimeHLSLOp.h"
+
+namespace NNX
+{
+	bool RegisterGemmOperator(FMLOperatorRegistryHlsl& Registry);
+} // NNX
