Operator registration and dummy test for HLSL runtime ConvTranspose operator

diff --git a/Engine/Plugins/Experimental/NNX/Source/NNXHlslShaders/Public/NNXConvTransposeCS.h b/Engine/Plugins/Experimental/NNX/Source/NNXHlslShaders/Public/NNXConvTransposeCS.h
--- a/Engine/Plugins/Experimental/NNX/Source/NNXHlslShaders/Public/NNXConvTransposeCS.h
+++ b/Engine/Plugins/Experimental/NNX/Source/NNXHlslShaders/Public/NNXConvTransposeCS.h
@@ -4,6 +4,7 @@
 
 #include "GlobalShader.h"
 #include "NNXShaderParameters.h"
+#include "NNXTypes.h"
 
 enum class EConvTransposeAlgorithm : uint8
 {
diff --git a/Engine/Plugins/Experimental/NNX/Source/NNXQA/Resources/NNXQATestDesc.json b/Engine/Plugins/Experimental/NNX/Source/NNXQA/Resources/NNXQATestDesc.json
--- a/Engine/Plugins/Experimental/NNX/Source/NNXQA/Resources/NNXQATestDesc.json
+++ b/Engine/Plugins/Experimental/NNX/Source/NNXQA/Resources/NNXQATestDesc.json
@@ -223,6 +223,13 @@
 			"targets": [
 				{ "target": "Gemm" }
 			]
+		},
+		{
+			"category": "ConvTranspose",
+			"runtimes": [	{	"name": "NNXRuntimeDml", "skip": true	}	],
+			"targets": [
+				{ "target": "ConvTranspose" }
+			]
 		}
 	],
 
@@ -355,6 +362,13 @@
 				{ "inputs": [ { "shape": [ 3, 4 ] }, { "shape": [ 4, 5 ] }, { "shape": [ 3, 5 ] } ], "outputs": [ { "shape": [ 3, 5 ] } ]},
 				{ "inputs": [ { "shape": [ 1, 1 ] }, { "shape": [ 1, 5 ] }, { "shape": [ 5 ] } ], "outputs": [ { "shape": [ 1, 5 ] } ]}
 			]
+		},
+		{
+			"name": "ConvTranspose",
+			"datasets": [
+				{	"inputs": [	{ "shape": [ 1, 1, 1 ] },	{ "shape": [ 1, 1, 1 ] } ],	"outputs": [ { "shape": [ 1, 1, 1 ] } ]},
+				{	"inputs": [	{ "shape": [ 1, 1, 1 ] },	{ "shape": [ 1, 1, 1 ] },	{ "shape": [ 1 ] } ],	"outputs": [ { "shape": [ 1, 1, 1 ] } ]}
+			]
 		}
 	]
 	
diff --git a/Engine/Plugins/Experimental/NNX/Source/NNXRuntimeRDG/Private/NNXRuntimeHLSL.cpp b/Engine/Plugins/Experimental/NNX/Source/NNXRuntimeRDG/Private/NNXRuntimeHLSL.cpp
--- a/Engine/Plugins/Experimental/NNX/Source/NNXRuntimeRDG/Private/NNXRuntimeHLSL.cpp
+++ b/Engine/Plugins/Experimental/NNX/Source/NNXRuntimeRDG/Private/NNXRuntimeHLSL.cpp
@@ -7,6 +7,7 @@
 #include "NNXRuntimeHLSLElementWiseUnaryOps.h"
 #include "NNXRuntimeHLSLElementWiseVariadicOps.h"
 #include "NNXRuntimeHLSLGemmOp.h"
+#include "NNXRuntimeHLSLConvTransposeOp.h"
 #include "NNXRuntimeHLSLHelper.h"
 #include "NNXRuntimeRDG.h"
 
@@ -153,6 +154,7 @@ public:
 		RegisterElementWiseBinaryOperators(*registry);
 		RegisterElementWiseVariadicOperators(*registry);
 		RegisterGemmOperator(*registry);
+		RegisterConvTransposeOperator(*registry);
 
 		return true;
 	}
diff --git a/Engine/Plugins/Experimental/NNX/Source/NNXRuntimeRDG/Private/NNXRuntimeHLSLConvTransposeOp.cpp b/Engine/Plugins/Experimental/NNX/Source/NNXRuntimeRDG/Private/NNXRuntimeHLSLConvTransposeOp.cpp
new file mode 100644
--- /dev/null
+++ b/Engine/Plugins/Experimental/NNX/Source/NNXRuntimeRDG/Private/NNXRuntimeHLSLConvTransposeOp.cpp
@@ -0,0 +1,127 @@
+// Copyright Epic Games, Inc. All Rights Reserved.
+
+#include "NNXRuntimeHLSLConvTransposeOp.h"
+#include "NNXConvTransposeCS.h"
+#include "NNXRuntimeHLSLHelper.h"
+
+namespace NNX
+{
+
+	/**
+	 * ConvTranspose ML operator implementation
+	 */
+	class FMLOperatorHlslConvTranspose : public FMLOperatorHlsl
+	{
+	public:
+
+		static FMLOperatorHlsl* Create()
+		{
+			return new FMLOperatorHlslConvTranspose();
+		}
+
+		virtual ~FMLOperatorHlslConvTranspose() = default;
+
+	private:
+
+		FMLOperatorHlslConvTranspose(){}
+
+		FMLTensorDesc Input;
+		FMLTensorDesc Weights;
+		FMLTensorDesc Bias;
+		FMLTensorDesc Output;
+
+		int NumDimensions;
+		bool HasBias;
+
+		// Hard coded parameters, until we accept them from JSON
+		int32 Group = 1;
+		TArray<int32> InputShape;
+		TArray<int32> WeightsShape;
+		EConvTransposeAutoPad AutoPad = EConvTransposeAutoPad::VALID;
+		TArray<int32> Dilations = {1};
+		TArray<int32> Strides = {1};
+		TArray<int32> Pads = {0, 0};
+		TArray<int32> OutputPadding = {0};
+
+	public:
+
+		virtual bool Initialize(TArrayView<const FMLTensorDesc> InputTensors, TArrayView<const FMLTensorDesc> OutputTensors) override
+		{
+			check(InputTensors.Num() >= 2 && InputTensors.Num() <= 3);
+			check(OutputTensors.Num() == 1);
+
+			Input = InputTensors[0];
+			Weights = InputTensors[1];
+			Output = OutputTensors[0];
+
+			check(Input.Dimension > 2);
+			check(Weights.Dimension  == Input.Dimension);
+			check(Output.Dimension == Input.Dimension);
+
+			if (InputTensors.Num() == 3) {
+				HasBias = true;
+				Bias = InputTensors[2];
+			}
+			else {
+				HasBias = false;
+			}
+
+			auto MakeShape = [](const NNX::FMLTensorDesc& InputDesc) {
+				TArray<int32> Result;
+				for (int32 i = 0; i < (int32)InputDesc.Dimension; i++) {
+					Result.Add(InputDesc.Sizes[i]);
+				}
+				return Result;
+			};
+			InputShape = MakeShape(Input);
+			WeightsShape = MakeShape(Weights);
+
+			NumDimensions = Input.Dimension - 2;
+
+			return true;
+		}
+
+		virtual void Dispatch(FRDGBuilder& GraphBuilder, TArrayView<const FMLTensorBinding> InInputBindings, TArrayView<const FMLTensorBinding> OutOutputBindings) override
+		{
+			constexpr EConvTransposeAlgorithm Algorithm = EConvTransposeAlgorithm::SharedMemory;
+			constexpr EConvTransposeGroupSize GroupSize = EConvTransposeGroupSize::Size256;
+
+			TArray<int32> OutputShape = FMLConvTransposeCS::GetOutputShape(InputShape, WeightsShape, AutoPad, Dilations, Strides, Pads, OutputPadding, Group);
+
+			// Set parameters
+			FMLConvTransposeCS::FParameters* Params = GraphBuilder.AllocParameters<FMLConvTransposeCS::FParameters>();
+			FMLConvTransposeCS::FillInParameters(GroupSize, InputShape, WeightsShape, HasBias, AutoPad, Group, Dilations,Strides, Pads, OutputPadding, *Params);
+			Params->X = GraphBuilder.CreateSRV(FRDGBufferSRVDesc(InInputBindings[0].Buffer, PF_R32_FLOAT));
+			Params->W = GraphBuilder.CreateSRV(FRDGBufferSRVDesc(InInputBindings[1].Buffer, PF_R32_FLOAT));
+			if (InInputBindings.Num() == 3) {
+				Params->B = GraphBuilder.CreateSRV(FRDGBufferSRVDesc(InInputBindings[2].Buffer, PF_R32_FLOAT));
+			}
+			Params->Y = GraphBuilder.CreateUAV(FRDGBufferUAVDesc(OutOutputBindings[0].Buffer, PF_R32_FLOAT));
+
+			FMLConvTransposeCS::FPermutationDomain PermutationVector;
+
+			PermutationVector.Set<FMLConvTransposeCS::FConvTransposeAlgorithm>(Algorithm);
+			PermutationVector.Set<FMLConvTransposeCS::FConvTransposeGroupSize>(GroupSize);
+			PermutationVector.Set<FMLConvTransposeCS::FConvTransposeNumStackDimensions>(NumDimensions);
+			PermutationVector.Set<FMLConvTransposeCS::FConvTransposeNumReadsPerThread>(FMLConvTransposeCS::GetNumReadsPerThread(GroupSize, WeightsShape, Dilations, Strides));
+			PermutationVector.Set<FMLConvTransposeCS::FConvTransposeHasB>(HasBias);
+			TShaderMapRef<FMLConvTransposeCS> ComputeShader(GetGlobalShaderMap(GMaxRHIFeatureLevel), PermutationVector);
+
+			FComputeShaderUtils::AddPass(
+				GraphBuilder,
+				RDG_EVENT_NAME("FMLOperatorHlslConvTranspose_Dispatch"),
+				ERDGPassFlags::Compute | ERDGPassFlags::NeverCull,
+				ComputeShader,
+				Params,
+				FMLConvTransposeCS::GetGroupCount(OutputShape, FMLConvTransposeCS::GetGroupShape(GroupSize, NumDimensions)));		
+		}
+	};
+
+	bool RegisterConvTransposeOperator(FMLOperatorRegistryHlsl& Registry)
+	{
+		Registry.OpAdd(TEXT("ConvTranspose"), FMLOperatorHlslConvTranspose::Create);
+
+		return true;
+	}
+
+} // NNX
diff --git a/Engine/Plugins/Experimental/NNX/Source/NNXRuntimeRDG/Private/NNXRuntimeHLSLConvTransposeOp.h b/Engine/Plugins/Experimental/NNX/Source/NNXRuntimeRDG/Private/NNXRuntimeHLSLConvTransposeOp.h
new file mode 100644
--- /dev/null
+++ b/Engine/Plugins/Experimental/NNX/Source/NNXRuntimeRDG/Private/NNXRuntimeHLSLConvTransposeOp.h
@@ -0,0 +1,11 @@
+// Copyright Epic Games, Inc. All Rights Reserved.
+
+#pragma once
+
+#include "NNXRuntimeRDG.h"
+#include "NNXRuntimeHLSLOp.h"
+
+namespace NNX
+{
+	bool RegisterConvTransposeOperator(FMLOperatorRegistryHlsl& Registry);
+} // NNX
