Operator registration and tests for HLSL runtime MatMul operator.

diff --git a/Engine/Plugins/Experimental/NNX/Source/NNXHlslShaders/Private/NNXGemmCS.cpp b/Engine/Plugins/Experimental/NNX/Source/NNXHlslShaders/Private/NNXGemmCS.cpp
--- a/Engine/Plugins/Experimental/NNX/Source/NNXHlslShaders/Private/NNXGemmCS.cpp
+++ b/Engine/Plugins/Experimental/NNX/Source/NNXHlslShaders/Private/NNXGemmCS.cpp
@@ -19,14 +19,14 @@ namespace
 		uint32 CHeight = 0;
 	};
 
-	static GemmMatrixParameters GetMatrixParameters(TArray<uint32> ShapeA, TArray<uint32> ShapeB, TArray<uint32> ShapeC)
+	static GemmMatrixParameters GetMatrixParameters(const NNX::FMLTensorDesc &InputA, const NNX::FMLTensorDesc &InputB, const NNX::FMLTensorDesc &InputC)
 	{
-		check(!ShapeA.IsEmpty());
-		check(!ShapeB.IsEmpty());
+		check(InputA.Dimension != 0);
+		check(InputB.Dimension != 0);
 
 		GemmMatrixParameters Result;
-		int32 NumStackDimensionsA = FMath::Max(ShapeA.Num() - 2, 0);
-		int32 NumStackDimensionsB = FMath::Max(ShapeB.Num() - 2, 0);
+		int32 NumStackDimensionsA = FMath::Max((int32)InputA.Dimension - 2, 0);
+		int32 NumStackDimensionsB = FMath::Max((int32)InputB.Dimension - 2, 0);
 		int32 NumStackDimensions = FMath::Max(NumStackDimensionsA, NumStackDimensionsB);
 
 		check(NumStackDimensions <= NNXRT_GEMM_MAX_NUM_STACK_DIMENSIONS);
@@ -37,29 +37,29 @@ namespace
 			const int32 NumDimToCheck = FMath::Min(NumStackDimensionsA, NumStackDimensionsB);
 			for (int32 i = 0; i < NumDimToCheck; i++)
 			{
-				const uint32 VolumeA = ShapeA[NumStackDimensionsA - 1 - i];
-				const uint32 VolumeB = ShapeB[NumStackDimensionsB - 1 - i];
+				const uint32 VolumeA = InputA.Sizes[NumStackDimensionsA - 1 - i];
+				const uint32 VolumeB = InputB.Sizes[NumStackDimensionsB - 1 - i];
 
 				check(VolumeA == 1 || VolumeB == 1 || VolumeA == VolumeB);
 			}
 		}
 
-		const bool IsVectorA = ShapeA.Num() == 1;
-		const bool IsVectorB = ShapeB.Num() == 1;
-		const bool IsVectorC = ShapeC.Num() == 1;
+		const bool IsVectorA = InputA.Dimension == 1;
+		const bool IsVectorB = InputB.Dimension == 1;
+		const bool IsVectorC = InputC.Dimension == 1;
 
-		Result.M = IsVectorA ? 1 : ShapeA[ShapeA.Num() - 2];
-		Result.N = IsVectorB ? 1 : ShapeB[ShapeB.Num() - 1];
-		Result.K = ShapeA[IsVectorA ? 0 : ShapeA.Num() - 1];
-		check(ShapeB[IsVectorB ? 0 : ShapeB.Num() - 2] == Result.K);
+		Result.M = IsVectorA ? 1 : InputA.Sizes[InputA.Dimension - 2];
+		Result.N = IsVectorB ? 1 : InputB.Sizes[InputB.Dimension - 1];
+		Result.K = InputA.Sizes[IsVectorA ? 0 : InputA.Dimension - 1];
+		check(InputB.Sizes[IsVectorB ? 0 : InputB.Dimension - 2] == Result.K);
 
 		Result.StackShapeA.Init(1, NumStackDimensions);
 		Result.StackShapeB.Init(1, NumStackDimensions);
 
 		for (int32 i = 0; i < NumStackDimensions; i++)
 		{
-			Result.StackShapeA[Result.StackShapeA.Num() - 1 - i] = i < NumStackDimensionsA ? ShapeA[ShapeA.Num() - 3 - i] : 1;
-			Result.StackShapeB[Result.StackShapeB.Num() - 1 - i] = i < NumStackDimensionsB ? ShapeB[ShapeB.Num() - 3 - i] : 1;
+			Result.StackShapeA[Result.StackShapeA.Num() - 1 - i] = i < NumStackDimensionsA ? InputA.Sizes[InputA.Dimension - 3 - i] : 1;
+			Result.StackShapeB[Result.StackShapeB.Num() - 1 - i] = i < NumStackDimensionsB ? InputB.Sizes[InputB.Dimension - 3 - i] : 1;
 		}
 
 		Result.StackStrideA.Init(1, NumStackDimensions);
@@ -71,8 +71,8 @@ namespace
 			Result.StackStrideB[i] = Result.StackStrideB[i + 1] * Result.StackShapeB[i + 1];
 		}
 
-		Result.CWidth = ShapeC.IsEmpty() ? 0 : ShapeC[ShapeC.Num() - 1];
-		Result.CHeight = ShapeC.IsEmpty() ? 0 : IsVectorC ? 1 : ShapeC[ShapeC.Num() - 2];
+		Result.CWidth = InputC.Dimension == 0 ? 0 : InputC.Sizes[InputC.Dimension - 1];
+		Result.CHeight = InputC.Dimension == 0 ? 0 : IsVectorC ? 1 : InputC.Sizes[InputC.Dimension - 2];
 
 		return Result;
 	}
@@ -87,15 +87,7 @@ void FMLGemmCS::ModifyCompilationEnvironment(const FGlobalShaderPermutationParam
 void FMLGemmCS::FillInParameters(float Alpha, float Beta, int32 TransA, int32 TransB, const NNX::FMLTensorDesc &InputA, const NNX::FMLTensorDesc &InputB,
 		const NNX::FMLTensorDesc &InputC, float CScalar, FMLGemmCS::FParameters& Parameters)
 {
-	auto MakeShape = [](const NNX::FMLTensorDesc &InputDesc) {
-		TArray<uint32> Result;
-		for (int32 i = 0; i < (int32)InputDesc.Dimension; i++) {
-			Result.Add(InputDesc.Sizes[i]);
-		}
-		return Result;
-	};
-
-	GemmMatrixParameters Params = GetMatrixParameters(MakeShape(InputA), MakeShape(InputB), MakeShape(InputC));
+	GemmMatrixParameters Params = GetMatrixParameters(InputA, InputB, InputC);
 
 	Parameters.Alpha = Alpha;
 	Parameters.Beta = Beta;
@@ -112,23 +104,9 @@ void FMLGemmCS::FillInParameters(float Alpha, float Beta, int32 TransA, int32 Tr
 	Parameters.CScalar = CScalar;
 }
 
-uint32 FMLGemmCS::GetShapeSize(TArray<uint32> Shape)
-{
-	auto Product = [](uint32 A, uint32 B) { return A * B; };
-
-	return Algo::Accumulate(Shape, 1, Product);
-}
-
-uint32 FMLGemmCS::GetMatMulOutputSize(TArray<uint32> ShapeA, TArray<uint32> ShapeB)
-{
-	GemmMatrixParameters Params = GetMatrixParameters(ShapeA, ShapeB, {});
-
-	return FMath::Max(GetShapeSize(Params.StackShapeA), GetShapeSize(Params.StackShapeB)) * Params.M * Params.N;
-}
-
-void FMLGemmCS::FillInParametersMatMul(TArray<uint32> ShapeA, TArray<uint32> ShapeB, FMLGemmCS::FParameters& Parameters)
+void FMLGemmCS::FillInParametersMatMul(const NNX::FMLTensorDesc &InputA, const NNX::FMLTensorDesc &InputB, FMLGemmCS::FParameters& Parameters)
 {
-	GemmMatrixParameters Params = GetMatrixParameters(ShapeA, ShapeB, {});
+	GemmMatrixParameters Params = GetMatrixParameters(InputA, InputB, {});
 
 	Parameters.Alpha = 1.0f;
 	Parameters.Beta = 1.0f;
diff --git a/Engine/Plugins/Experimental/NNX/Source/NNXHlslShaders/Public/NNXGemmCS.h b/Engine/Plugins/Experimental/NNX/Source/NNXHlslShaders/Public/NNXGemmCS.h
--- a/Engine/Plugins/Experimental/NNX/Source/NNXHlslShaders/Public/NNXGemmCS.h
+++ b/Engine/Plugins/Experimental/NNX/Source/NNXHlslShaders/Public/NNXGemmCS.h
@@ -52,10 +52,7 @@ public:
 	static void FillInParameters(float Alpha, float Beta, int32 TransA, int32 TransB, const NNX::FMLTensorDesc &InputA, const NNX::FMLTensorDesc &InputB,
 		const NNX::FMLTensorDesc &InputC, float CScalar, FMLGemmCS::FParameters& Parameters);
 
-	static uint32 GetShapeSize(TArray<uint32> Shape);
-	static uint32 GetMatMulOutputSize(TArray<uint32> ShapeA, TArray<uint32> ShapeB);
-
-	static void FillInParametersMatMul(TArray<uint32> ShapeA, TArray<uint32> ShapeB, FMLGemmCS::FParameters& Parameters);
+	static void FillInParametersMatMul(const NNX::FMLTensorDesc &InputA, const NNX::FMLTensorDesc &InputB, FMLGemmCS::FParameters& Parameters);
 
 	static FIntVector GetGroupCount(const FMLGemmCS::FParameters& Parameters, EGemmAlgorithm Algorithm, int32 NumStackDimensions);
 	static EGemmAlgorithm GetAlgorithm(const FMLGemmCS::FParameters& Parameters);
diff --git a/Engine/Plugins/Experimental/NNX/Source/NNXQA/Resources/NNXQATestDesc.json b/Engine/Plugins/Experimental/NNX/Source/NNXQA/Resources/NNXQATestDesc.json
--- a/Engine/Plugins/Experimental/NNX/Source/NNXQA/Resources/NNXQATestDesc.json
+++ b/Engine/Plugins/Experimental/NNX/Source/NNXQA/Resources/NNXQATestDesc.json
@@ -230,6 +230,13 @@
 			"targets": [
 				{ "target": "ConvTranspose" }
 			]
+		},
+		{
+			"category": "MatMul",
+			"runtimes": [ { "name": "NNXRuntimeDml", "skip": true } ],
+			"targets": [
+				{ "target": "MatMul" }
+			]
 		}
 	],
 
@@ -369,6 +376,27 @@
 				{	"inputs": [	{ "shape": [ 1, 1, 1 ] },	{ "shape": [ 1, 1, 1 ] } ],	"outputs": [ { "shape": [ 1, 1, 1 ] } ]},
 				{	"inputs": [	{ "shape": [ 1, 1, 1 ] },	{ "shape": [ 1, 1, 1 ] },	{ "shape": [ 1 ] } ],	"outputs": [ { "shape": [ 1, 1, 1 ] } ]}
 			]
+		},
+		{
+			"name": "MatMul",
+			"datasets": [
+				{ "inputs": [ { "shape": [ 2, 2 ] }, { "shape": [ 2, 2 ] } ], "outputs": [ { "shape": [ 2, 2 ] } ]},
+				{ "inputs": [ { "shape": [ 3, 4 ] }, { "shape": [ 4, 5 ] } ], "outputs": [ { "shape": [ 3, 5 ] } ]},
+				{ "inputs": [ { "shape": [ 2, 3, 4 ] }, { "shape": [ 2, 4, 5 ] } ], "outputs": [ { "shape": [ 2, 3, 5 ] } ]},
+				{ "inputs": [ { "shape": [ 5, 4 ] }, { "shape": [ 2, 4, 3 ] } ], "outputs": [ { "shape": [ 2, 5, 3 ] } ]},
+				{ "inputs": [ { "shape": [ 2, 5, 4 ] }, { "shape": [ 4, 3 ] } ], "outputs": [ { "shape": [ 2, 5, 3 ] } ]},
+				{ "inputs": [ { "shape": [ 1, 4 ] }, { "shape": [ 4, 3 ] } ], "outputs": [ { "shape": [ 1, 3 ] } ]},
+				{ "inputs": [ { "shape": [ 1, 4 ] }, { "shape": [ 2, 4, 3 ] } ], "outputs": [ { "shape": [ 2, 1, 3 ] } ]},
+				{ "inputs": [ { "shape": [ 3, 4 ] }, { "shape": [ 4, 1 ] } ], "outputs": [ { "shape": [ 3, 1 ] } ]},
+				{ "inputs": [ { "shape": [ 2, 2, 3, 4 ] }, { "shape": [ 4, 1 ] } ], "outputs": [ { "shape": [ 2, 2, 3, 1 ] } ]},
+				{ "inputs": [ { "shape": [ 1, 6 ] }, { "shape": [ 6, 1 ] } ], "outputs": [ { "shape": [ 1, 1 ] } ]},
+				{ "inputs": [ { "shape": [ 2, 2, 3, 5, 4 ] }, { "shape": [ 2, 1, 3, 4, 3 ] } ], "outputs": [ { "shape": [ 2, 2, 3, 5, 3 ] } ]},
+				{ "inputs": [ { "shape": [ 2, 1, 2, 5, 4 ] }, { "shape": [ 1, 3, 2, 4, 3 ] } ], "outputs": [ { "shape": [ 2, 3, 2, 5, 3 ] } ]},
+				{ "inputs": [ { "shape": [ 1, 1, 2, 5, 4 ] }, { "shape": [ 2, 3, 2, 4, 3 ] } ], "outputs": [ { "shape": [ 2, 3, 2, 5, 3 ] } ]},
+				{ "inputs": [ { "shape": [ 1, 2, 5, 4 ] }, { "shape": [ 2, 3, 2, 4, 3 ] } ], "outputs": [ { "shape": [ 2, 3, 2, 5, 3 ] } ]},
+				{ "inputs": [ { "shape": [ 5, 4 ] }, { "shape": [ 2, 1, 3, 4, 3 ] } ], "outputs": [ { "shape": [ 2, 1, 3, 5, 3 ] } ]},
+				{ "inputs": [ { "shape": [ 2, 1, 3, 5, 4 ] }, { "shape": [ 4, 3 ] } ], "outputs": [ { "shape": [ 2, 1, 3, 5, 3 ] } ]}
+			]
 		}
 	]
 	
diff --git a/Engine/Plugins/Experimental/NNX/Source/NNXRuntimeRDG/Private/NNXRuntimeHLSL.cpp b/Engine/Plugins/Experimental/NNX/Source/NNXRuntimeRDG/Private/NNXRuntimeHLSL.cpp
--- a/Engine/Plugins/Experimental/NNX/Source/NNXRuntimeRDG/Private/NNXRuntimeHLSL.cpp
+++ b/Engine/Plugins/Experimental/NNX/Source/NNXRuntimeRDG/Private/NNXRuntimeHLSL.cpp
@@ -8,6 +8,7 @@
 #include "NNXRuntimeHLSLElementWiseVariadicOps.h"
 #include "NNXRuntimeHLSLGemmOp.h"
 #include "NNXRuntimeHLSLConvTransposeOp.h"
+#include "NNXRuntimeHLSLMatMulOp.h"
 #include "NNXRuntimeHLSLHelper.h"
 #include "NNXRuntimeRDG.h"
 
@@ -155,6 +156,7 @@ public:
 		RegisterElementWiseVariadicOperators(*registry);
 		RegisterGemmOperator(*registry);
 		RegisterConvTransposeOperator(*registry);
+		RegisterMatMulOperator(*registry);
 
 		return true;
 	}
diff --git a/Engine/Plugins/Experimental/NNX/Source/NNXRuntimeRDG/Private/NNXRuntimeHLSLMatMulOp.cpp b/Engine/Plugins/Experimental/NNX/Source/NNXRuntimeRDG/Private/NNXRuntimeHLSLMatMulOp.cpp
new file mode 100644
--- /dev/null
+++ b/Engine/Plugins/Experimental/NNX/Source/NNXRuntimeRDG/Private/NNXRuntimeHLSLMatMulOp.cpp
@@ -0,0 +1,84 @@
+// Copyright Epic Games, Inc. All Rights Reserved.
+
+#include "NNXRuntimeHLSLMatMulOp.h"
+#include "NNXGemmCS.h"
+#include "NNXRuntimeHLSLHelper.h"
+
+namespace NNX
+{
+
+	/**
+	 * Gemm ML operator implementation
+	 */
+	class FMLOperatorHlslMatMul : public FMLOperatorHlsl
+	{
+	public:
+
+		static FMLOperatorHlsl* Create()
+		{
+			return new FMLOperatorHlslMatMul();
+		}
+
+		virtual ~FMLOperatorHlslMatMul() = default;
+
+	private:
+
+		FMLOperatorHlslMatMul() {}
+
+		FMLTensorDesc InputA = {};
+		FMLTensorDesc InputB = {};
+		FMLTensorDesc Output = {};
+
+	public:
+
+		virtual bool Initialize(TArrayView<const FMLTensorDesc> InputTensors, TArrayView<const FMLTensorDesc> OutputTensors) override
+		{
+			check(InputTensors.Num() == 2);
+			check(OutputTensors.Num() == 1);
+
+			InputA = InputTensors[0];
+			InputB = InputTensors[1];
+			Output = OutputTensors[0];
+
+			return true;
+		}
+
+		virtual void Dispatch(FRDGBuilder& GraphBuilder, TArrayView<const FMLTensorBinding> InInputBindings, TArrayView<const FMLTensorBinding> OutOutputBindings) override
+		{
+			const EGemmAlgorithm Algorithm = EGemmAlgorithm::Simple32x32;
+
+			const int32 NumStackDimensions = FMath::Max((int32)FMath::Max(InputA.Dimension, InputB.Dimension) - 2, 0);
+
+			// Set parameters
+			FMLGemmCS::FParameters* Parameters = GraphBuilder.AllocParameters<FMLGemmCS::FParameters>();
+			FMLGemmCS::FillInParametersMatMul(InputA, InputB, *Parameters);
+			Parameters->A = GraphBuilder.CreateSRV(FRDGBufferSRVDesc(InInputBindings[0].Buffer, PF_R32_FLOAT));
+			Parameters->B = GraphBuilder.CreateSRV(FRDGBufferSRVDesc(InInputBindings[1].Buffer, PF_R32_FLOAT));
+			Parameters->Y = GraphBuilder.CreateUAV(FRDGBufferUAVDesc(OutOutputBindings[0].Buffer, PF_R32_FLOAT));
+
+			FMLGemmCS::FPermutationDomain PermutationVector;
+			PermutationVector.Set<FMLGemmCS::FGemmCScalar>(EGemmCScalar::NoBias);
+			PermutationVector.Set<FMLGemmCS::FGemmAlgorithm>(Algorithm);
+			PermutationVector.Set<FMLGemmCS::FGemmNumStackDimensions>(NumStackDimensions);
+			TShaderMapRef<FMLGemmCS> ComputeShader(GetGlobalShaderMap(GMaxRHIFeatureLevel), PermutationVector);
+
+			FIntVector ThreadGroupCount = FMLGemmCS::GetGroupCount(*Parameters, Algorithm, NumStackDimensions);
+
+			FComputeShaderUtils::AddPass(
+				GraphBuilder,
+				RDG_EVENT_NAME("FMLHlslMatMulOperatorHlsl_Dispatch"),
+				ERDGPassFlags::Compute | ERDGPassFlags::NeverCull,
+				ComputeShader,
+				Parameters,
+				ThreadGroupCount);
+		}
+	};
+
+	bool RegisterMatMulOperator(FMLOperatorRegistryHlsl& Registry)
+	{
+		Registry.OpAdd(TEXT("MatMul"), FMLOperatorHlslMatMul::Create);
+
+		return true;
+	}
+
+} // NNX
diff --git a/Engine/Plugins/Experimental/NNX/Source/NNXRuntimeRDG/Private/NNXRuntimeHLSLMatMulOp.h b/Engine/Plugins/Experimental/NNX/Source/NNXRuntimeRDG/Private/NNXRuntimeHLSLMatMulOp.h
new file mode 100644
--- /dev/null
+++ b/Engine/Plugins/Experimental/NNX/Source/NNXRuntimeRDG/Private/NNXRuntimeHLSLMatMulOp.h
@@ -0,0 +1,11 @@
+// Copyright Epic Games, Inc. All Rights Reserved.
+
+#pragma once
+
+#include "NNXRuntimeRDG.h"
+#include "NNXRuntimeHLSLOp.h"
+
+namespace NNX
+{
+	bool RegisterMatMulOperator(FMLOperatorRegistryHlsl& Registry);
+} // NNX
