Refactor uniforms of binary and variadic operators, to fix Linux shader problems and allow for indefinite number of dimensions

#rnx

diff --git a/Engine/Plugins/Experimental/NNX/Source/NNXHlslShaders/Public/NNXElementWiseBinaryCS.h b/Engine/Plugins/Experimental/NNX/Source/NNXHlslShaders/Public/NNXElementWiseBinaryCS.h
--- a/Engine/Plugins/Experimental/NNX/Source/NNXHlslShaders/Public/NNXElementWiseBinaryCS.h
+++ b/Engine/Plugins/Experimental/NNX/Source/NNXHlslShaders/Public/NNXElementWiseBinaryCS.h
@@ -22,7 +22,8 @@ public:
 	static const uint32 THREADGROUP_SIZE_X;
 
 	class FOperatorType : SHADER_PERMUTATION_ENUM_CLASS("OP_TYPENAME", EMLElementWiseBinaryOperatorType);
-	using FPermutationDomain = TShaderPermutationDomain<FOperatorType>;
+	class FBinaryNumDimensions : SHADER_PERMUTATION_RANGE_INT("NUM_DIMENSIONS", 1, NXRT_TENSORSTRIDEINFO_MAX_NUM_DIMENSIONS);
+	using FPermutationDomain = TShaderPermutationDomain<FOperatorType, FBinaryNumDimensions>;
 
 	static void ModifyCompilationEnvironment(const FGlobalShaderPermutationParameters& InParameters, FShaderCompilerEnvironment& OutEnvironment);
 
diff --git a/Engine/Plugins/Experimental/NNX/Source/NNXHlslShaders/Public/NNXElementWiseVariadicCS.h b/Engine/Plugins/Experimental/NNX/Source/NNXHlslShaders/Public/NNXElementWiseVariadicCS.h
--- a/Engine/Plugins/Experimental/NNX/Source/NNXHlslShaders/Public/NNXElementWiseVariadicCS.h
+++ b/Engine/Plugins/Experimental/NNX/Source/NNXHlslShaders/Public/NNXElementWiseVariadicCS.h
@@ -21,7 +21,8 @@ public:
 	class FApplyScale : SHADER_PERMUTATION_BOOL("APPLYSCALE");
 	class FOutputAsInput : SHADER_PERMUTATION_BOOL("OUTPUTASINPUT");
 	class FNumInput : SHADER_PERMUTATION_RANGE_INT("NUMINPUT", 1, MAX_NUM_INPUT);
-	using FPermutationDomain = TShaderPermutationDomain<FOperatorType, FApplyScale, FOutputAsInput, FNumInput>;
+	class FVariadicNumDimensions : SHADER_PERMUTATION_RANGE_INT("NUM_DIMENSIONS", 1, NXRT_TENSORSTRIDEINFO_MAX_NUM_DIMENSIONS);
+	using FPermutationDomain = TShaderPermutationDomain<FOperatorType, FApplyScale, FOutputAsInput, FNumInput, FVariadicNumDimensions>;
 
 	static void ModifyCompilationEnvironment(const FGlobalShaderPermutationParameters& InParameters, FShaderCompilerEnvironment& OutEnvironment);
 
diff --git a/Engine/Plugins/Experimental/NNX/Source/NNXHlslShaders/Public/NNXShaderParameters.h b/Engine/Plugins/Experimental/NNX/Source/NNXHlslShaders/Public/NNXShaderParameters.h
--- a/Engine/Plugins/Experimental/NNX/Source/NNXHlslShaders/Public/NNXShaderParameters.h
+++ b/Engine/Plugins/Experimental/NNX/Source/NNXHlslShaders/Public/NNXShaderParameters.h
@@ -22,6 +22,8 @@
 
 #define NNXRT_END_SHADER_PARAMETER_STRUCT END_SHADER_PARAMETER_STRUCT
 
+#define NXRT_TENSORSTRIDEINFO_MAX_NUM_DIMENSIONS 8
+
 //Unary element wise operators
 #define NNXRT_ELEMENTWISEUNARY_PARAMETER_STRUCT() \
 	NNXRT_BEGIN_SHADER_PARAMETER_STRUCT(FMLElementWiseUnaryParameters) \
@@ -46,17 +48,8 @@ NNXRT_ELEMENTWISEUNARY_PARAMETER_STRUCT()
 		NNXRT_RDG_BUFFER_SRV(Buffer<float>, Input2) \
 		NNXRT_RDG_BUFFER_SRV(Buffer<float>, Input3) \
 		NNXRT_RDG_BUFFER_UAV(RWBuffer<float>, Output) \
-		SHADER_PARAMETER(FUint32Vector4, Input0Info0) \
-		SHADER_PARAMETER(FUint32Vector4, Input0Info1) \
-		SHADER_PARAMETER(FUint32Vector4, Input1Info0) \
-		SHADER_PARAMETER(FUint32Vector4, Input1Info1) \
-		SHADER_PARAMETER(FUint32Vector4, Input2Info0) \
-		SHADER_PARAMETER(FUint32Vector4, Input2Info1) \
-		SHADER_PARAMETER(FUint32Vector4, Input3Info0) \
-		SHADER_PARAMETER(FUint32Vector4, Input3Info1) \
-		SHADER_PARAMETER(FUint32Vector4, OutInfo0) \
-		SHADER_PARAMETER(FUint32Vector4, OutInfo1) \
-		SHADER_PARAMETER(uint32, OutRank) \
+		SHADER_PARAMETER_ARRAY(FUintVector4, InputTensorInfo, [NXRT_TENSORSTRIDEINFO_MAX_NUM_DIMENSIONS]) \
+		SHADER_PARAMETER_ARRAY(FUintVector4, OutputTensorInfo, [NXRT_TENSORSTRIDEINFO_MAX_NUM_DIMENSIONS]) \
 		SHADER_PARAMETER(uint32, Num) \
 		SHADER_PARAMETER(uint32, ThreadCountX) \
 		SHADER_PARAMETER(float, Scale) \
@@ -72,13 +65,7 @@ NNXRT_ELEMENTWISEVARIADIC_PARAMETER_STRUCT()
 		NNXRT_RDG_BUFFER_SRV(Buffer<float>, LHSInput) \
 		NNXRT_RDG_BUFFER_SRV(Buffer<float>, RHSInput) \
 		NNXRT_RDG_BUFFER_UAV(RWBuffer<float>, Output) \
-		SHADER_PARAMETER(FUint32Vector4, LHSInfo0) \
-		SHADER_PARAMETER(FUint32Vector4, LHSInfo1) \
-		SHADER_PARAMETER(FUint32Vector4, RHSInfo0) \
-		SHADER_PARAMETER(FUint32Vector4, RHSInfo1) \
-		SHADER_PARAMETER(FUint32Vector4, OutInfo0) \
-		SHADER_PARAMETER(FUint32Vector4, OutInfo1) \
-		SHADER_PARAMETER(uint32, OutRank) \
+		SHADER_PARAMETER_ARRAY(FUintVector4, TensorInfo, [NXRT_TENSORSTRIDEINFO_MAX_NUM_DIMENSIONS]) \
 		SHADER_PARAMETER(uint32, Num) \
 		SHADER_PARAMETER(uint32, ThreadCountX) \
 	NNXRT_END_SHADER_PARAMETER_STRUCT()
diff --git a/Engine/Plugins/Experimental/NNX/Source/NNXHlslShaders/Shaders/BroadcastHelper.ush b/Engine/Plugins/Experimental/NNX/Source/NNXHlslShaders/Shaders/BroadcastHelper.ush
--- a/Engine/Plugins/Experimental/NNX/Source/NNXHlslShaders/Shaders/BroadcastHelper.ush
+++ b/Engine/Plugins/Experimental/NNX/Source/NNXHlslShaders/Shaders/BroadcastHelper.ush
@@ -1,28 +1,5 @@
 // Copyright Epic Games, Inc. All Rights Reserved.
 
-// Contain the tensor strides padded with 0s
-// tensor shape [2,3,4] have strides [12,4,1]
-// and is stored as:
-// PaddedStrides = {12,4,1,0,0,0,0,0};
-struct FTensorInfo
-{
-	uint PaddedStrides[8];
-};
-
-struct FTensorInfo DecodeTensorInfo(uint4 Info0, uint4 Info1)
-{
-	FTensorInfo Out;
-	Out.PaddedStrides[0] = Info0[0];
-	Out.PaddedStrides[1] = Info0[1];
-	Out.PaddedStrides[2] = Info0[2];
-	Out.PaddedStrides[3] = Info0[3];
-	Out.PaddedStrides[4] = Info1[0];
-	Out.PaddedStrides[5] = Info1[1];
-	Out.PaddedStrides[6] = Info1[2];
-	Out.PaddedStrides[7] = Info1[3];
-	return Out;
-}
-
 void DivMod(uint op1, uint op2, out uint q, out uint r)
 {
 	q = op1 / op2;
diff --git a/Engine/Plugins/Experimental/NNX/Source/NNXHlslShaders/Shaders/ElementWiseBinaryOp.usf b/Engine/Plugins/Experimental/NNX/Source/NNXHlslShaders/Shaders/ElementWiseBinaryOp.usf
--- a/Engine/Plugins/Experimental/NNX/Source/NNXHlslShaders/Shaders/ElementWiseBinaryOp.usf
+++ b/Engine/Plugins/Experimental/NNX/Source/NNXHlslShaders/Shaders/ElementWiseBinaryOp.usf
@@ -6,13 +6,7 @@
 Buffer<float> LHSInput;
 Buffer<float> RHSInput;
 RWBuffer<float> Output;
-uint4 LHSInfo0;
-uint4 LHSInfo1;
-uint4 RHSInfo0;
-uint4 RHSInfo1;
-uint4 OutInfo0;
-uint4 OutInfo1;
-uint OutRank;
+uint4 TensorInfo[NUM_DIMENSIONS];
 uint Num;
 uint ThreadCountX;
 
@@ -29,19 +23,15 @@ void ElementWiseBinaryOp(in const uint3 DispatchThreadID : SV_DispatchThreadID)
 	
 	if (Index < Num)
 	{
-	    const FTensorInfo LHSInfo = DecodeTensorInfo(LHSInfo0, LHSInfo1);
-		const FTensorInfo RHSInfo = DecodeTensorInfo(RHSInfo0, RHSInfo1);
-		const FTensorInfo OutInfo = DecodeTensorInfo(OutInfo0, OutInfo1);
-
 		uint LHSIndex = 0;
 		uint RHSIndex = 0;
 		uint Offset = Index;
-		for (uint dim = 0; dim < OutRank; ++dim)
+		for (uint dim = 0; dim < NUM_DIMENSIONS; ++dim)
 		{
 			uint Q, R;
-			DivMod(Offset, OutInfo.PaddedStrides[dim], Q, R);
-			LHSIndex += LHSInfo.PaddedStrides[dim] * Q;
-			RHSIndex += RHSInfo.PaddedStrides[dim] * Q;
+			DivMod(Offset, TensorInfo[dim][2], Q, R);
+			LHSIndex += TensorInfo[dim][0] * Q;
+			RHSIndex += TensorInfo[dim][1] * Q;
 			Offset = R;
 		}
 
diff --git a/Engine/Plugins/Experimental/NNX/Source/NNXHlslShaders/Shaders/ElementWiseVariadicOp.usf b/Engine/Plugins/Experimental/NNX/Source/NNXHlslShaders/Shaders/ElementWiseVariadicOp.usf
--- a/Engine/Plugins/Experimental/NNX/Source/NNXHlslShaders/Shaders/ElementWiseVariadicOp.usf
+++ b/Engine/Plugins/Experimental/NNX/Source/NNXHlslShaders/Shaders/ElementWiseVariadicOp.usf
@@ -10,17 +10,8 @@ Buffer<BUFFER_TYPE> Input1;
 Buffer<BUFFER_TYPE> Input2;
 Buffer<BUFFER_TYPE> Input3;
 RWBuffer<BUFFER_TYPE> Output;
-uint4 Input0Info0;
-uint4 Input0Info1;
-uint4 Input1Info0;
-uint4 Input1Info1;
-uint4 Input2Info0;
-uint4 Input2Info1;
-uint4 Input3Info0;
-uint4 Input3Info1;
-uint4 OutInfo0;
-uint4 OutInfo1;
-uint OutRank;
+uint4 InputTensorInfo[NUM_DIMENSIONS];
+uint4 OutputTensorInfo[NUM_DIMENSIONS];
 uint Num;
 uint ThreadCountX;
 float Scale;
@@ -32,25 +23,19 @@ void ElementWiseVariadicOp(in const uint3 DispatchThreadID : SV_DispatchThreadID
 	
 	if (Index < Num)
 	{
-	    const FTensorInfo Info0 = DecodeTensorInfo(Input0Info0, Input0Info1);
-		const FTensorInfo Info1 = DecodeTensorInfo(Input1Info0, Input1Info1);
-		const FTensorInfo Info2 = DecodeTensorInfo(Input2Info0, Input2Info1);
-		const FTensorInfo Info3 = DecodeTensorInfo(Input3Info0, Input3Info1);
-		const FTensorInfo OutInfo = DecodeTensorInfo(OutInfo0, OutInfo1);
-
 		uint Input0Index = 0;
 		uint Input1Index = 0;
 		uint Input2Index = 0;
 		uint Input3Index = 0;
 		uint Offset = Index;
-		for (uint dim = 0; dim < OutRank; ++dim)
+		for (uint dim = 0; dim < NUM_DIMENSIONS; ++dim)
 		{
 			uint Q, R;
-			DivMod(Offset, OutInfo.PaddedStrides[dim], Q, R);
-			Input0Index += Info0.PaddedStrides[dim] * Q;
-			Input1Index += Info1.PaddedStrides[dim] * Q;
-			Input2Index += Info2.PaddedStrides[dim] * Q;
-			Input3Index += Info3.PaddedStrides[dim] * Q;
+			DivMod(Offset, OutputTensorInfo[dim][0], Q, R);
+			Input0Index += InputTensorInfo[dim][0] * Q;
+			Input1Index += InputTensorInfo[dim][1] * Q;
+			Input2Index += InputTensorInfo[dim][2] * Q;
+			Input3Index += InputTensorInfo[dim][3] * Q;
 			Offset = R;
 		}
 
diff --git a/Engine/Plugins/Experimental/NNX/Source/NNXRuntimeRDG/Private/NNXRuntimeHLSLElementWiseBinaryOps.cpp b/Engine/Plugins/Experimental/NNX/Source/NNXRuntimeRDG/Private/NNXRuntimeHLSLElementWiseBinaryOps.cpp
--- a/Engine/Plugins/Experimental/NNX/Source/NNXRuntimeRDG/Private/NNXRuntimeHLSLElementWiseBinaryOps.cpp
+++ b/Engine/Plugins/Experimental/NNX/Source/NNXRuntimeRDG/Private/NNXRuntimeHLSLElementWiseBinaryOps.cpp
@@ -58,16 +58,16 @@ public:
 		Params->LHSInput = LHSInputSRV;
 		Params->RHSInput = RHSInputSRV;
 		Params->Output = OutputUAV;
-		FillTensorStrideForBroadcastShaderParameters(LHSInput, Output.Dimension, Params->LHSInfo0, Params->LHSInfo1);
-		FillTensorStrideForBroadcastShaderParameters(RHSInput, Output.Dimension, Params->RHSInfo0, Params->RHSInfo1);
-		FillTensorStrideShaderParameters(Output, Params->OutInfo0, Params->OutInfo1);
-		Params->OutRank = Output.Dimension;
+		FillTensorStrideForBroadcastShaderParameters(LHSInput, Output.Dimension, Params->TensorInfo, 0);
+		FillTensorStrideForBroadcastShaderParameters(RHSInput, Output.Dimension, Params->TensorInfo, 1);
+		FillTensorStrideShaderParameters(Output, Params->TensorInfo, 2);
 		Params->Num = Output.Num();
 		Params->ThreadCountX = ThreadGroupCount.X * FMLElementWiseBinaryCS::THREADGROUP_SIZE_X;
 
 		FMLElementWiseBinaryCS::FPermutationDomain PermutationVector;
 
 		PermutationVector.Set<FMLElementWiseBinaryCS::FOperatorType>(OpType);
+		PermutationVector.Set<FMLElementWiseBinaryCS::FBinaryNumDimensions>(Output.Dimension);
 
 		TShaderMapRef<FMLElementWiseBinaryCS> ComputeShader(GetGlobalShaderMap(GMaxRHIFeatureLevel), PermutationVector);
 
diff --git a/Engine/Plugins/Experimental/NNX/Source/NNXRuntimeRDG/Private/NNXRuntimeHLSLElementWiseVariadicOps.cpp b/Engine/Plugins/Experimental/NNX/Source/NNXRuntimeRDG/Private/NNXRuntimeHLSLElementWiseVariadicOps.cpp
--- a/Engine/Plugins/Experimental/NNX/Source/NNXRuntimeRDG/Private/NNXRuntimeHLSLElementWiseVariadicOps.cpp
+++ b/Engine/Plugins/Experimental/NNX/Source/NNXRuntimeRDG/Private/NNXRuntimeHLSLElementWiseVariadicOps.cpp
@@ -40,21 +40,20 @@ namespace NNX
 		Params->Input2 = InputsSRV[2];
 		Params->Input3 = InputsSRV[3];
 		Params->Output = OutputUAV;
-		FillTensorStrideForBroadcastShaderParameters(InputDesc[0], OutputDesc.Dimension, Params->Input0Info0, Params->Input0Info1);
+		FillTensorStrideForBroadcastShaderParameters(InputDesc[0], OutputDesc.Dimension, Params->InputTensorInfo, 0);
 		if (InputBindings.Num() >= 2)
 		{
-			FillTensorStrideForBroadcastShaderParameters(InputDesc[1], OutputDesc.Dimension, Params->Input1Info0, Params->Input1Info1);
+			FillTensorStrideForBroadcastShaderParameters(InputDesc[1], OutputDesc.Dimension, Params->InputTensorInfo, 1);
 		}
 		if (InputBindings.Num() >= 3)
 		{
-			FillTensorStrideForBroadcastShaderParameters(InputDesc[2], OutputDesc.Dimension, Params->Input2Info0, Params->Input2Info1);
+			FillTensorStrideForBroadcastShaderParameters(InputDesc[2], OutputDesc.Dimension, Params->InputTensorInfo, 2);
 		}
 		if (InputBindings.Num() >= 4)
 		{
-			FillTensorStrideForBroadcastShaderParameters(InputDesc[3], OutputDesc.Dimension, Params->Input3Info0, Params->Input3Info1);
+			FillTensorStrideForBroadcastShaderParameters(InputDesc[3], OutputDesc.Dimension, Params->InputTensorInfo, 3);
 		}
-		FillTensorStrideShaderParameters(OutputDesc, Params->OutInfo0, Params->OutInfo1);
-		Params->OutRank = OutputDesc.Dimension;
+		FillTensorStrideShaderParameters(OutputDesc, Params->OutputTensorInfo, 0);
 		Params->Num = OutputDesc.Num();
 		Params->ThreadCountX = ThreadGroupCount.X * FMLElementWiseVariadicCS::THREADGROUP_SIZE_X;
 		Params->Scale = Scale;
@@ -66,6 +65,7 @@ namespace NNX
 		PermutationVector.Set<FMLElementWiseVariadicCS::FApplyScale>(Scale != 1.0f);
 		PermutationVector.Set<FMLElementWiseVariadicCS::FOutputAsInput>(OutputAsInput);
 		PermutationVector.Set<FMLElementWiseVariadicCS::FNumInput>(InputBindings.Num());
+		PermutationVector.Set<FMLElementWiseVariadicCS::FVariadicNumDimensions>(OutputDesc.Dimension);
 
 		// Add the pass to RDG
 		TShaderMapRef<FMLElementWiseVariadicCS> ComputeShader(GetGlobalShaderMap(GMaxRHIFeatureLevel), PermutationVector);
diff --git a/Engine/Plugins/Experimental/NNX/Source/NNXRuntimeRDG/Private/NNXRuntimeHLSLHelper.cpp b/Engine/Plugins/Experimental/NNX/Source/NNXRuntimeRDG/Private/NNXRuntimeHLSLHelper.cpp
--- a/Engine/Plugins/Experimental/NNX/Source/NNXRuntimeRDG/Private/NNXRuntimeHLSLHelper.cpp
+++ b/Engine/Plugins/Experimental/NNX/Source/NNXRuntimeRDG/Private/NNXRuntimeHLSLHelper.cpp
@@ -6,7 +6,7 @@
 
 namespace NNX
 {
-	void ComputeTensorStrides(const FMLTensorDesc& TensorDesc, uint32 OutStrides[8], uint32 TargetNumdimensionForBroadcast)
+	void FillTensorStrideShaderParameters(const FMLTensorDesc& TensorDesc, TStaticArray<FUintVector4, NXRT_TENSORSTRIDEINFO_MAX_NUM_DIMENSIONS, 16U>& OutShaderParam, uint32 Idx, uint32 TargetNumdimensionForBroadcast)
 	{
 		if (TargetNumdimensionForBroadcast == -1)
 		{
@@ -15,62 +15,39 @@ namespace NNX
 		checkf(TargetNumdimensionForBroadcast >= TensorDesc.Dimension, TEXT("Can't broadcast tensor from rank %d to rank %d, should be inferior or equal."), TensorDesc.Dimension, TargetNumdimensionForBroadcast);
 		uint32 Offset = TargetNumdimensionForBroadcast - TensorDesc.Dimension;
 
-		static_assert(FMLTensorDesc::MaxTensorDimension <= 8);
+		static_assert(FMLTensorDesc::MaxTensorDimension <= NXRT_TENSORSTRIDEINFO_MAX_NUM_DIMENSIONS);
 		for (int32 i = 7; i >= 0; --i)
 		{
 			if ((uint32)i >= TargetNumdimensionForBroadcast || (uint32)i < Offset)
 			{
-				OutStrides[i] = 0;
+				OutShaderParam[i][Idx] = 0;
 			}
 			else if ((uint32)i == TargetNumdimensionForBroadcast - 1)
 			{
-				OutStrides[i] = 1;
+				OutShaderParam[i][Idx] = 1;
 			}
 			else
 			{
-				OutStrides[i] = OutStrides[i + 1] * TensorDesc.Sizes[i + 1 - Offset];
+				OutShaderParam[i][Idx] = OutShaderParam[i + 1][Idx] * TensorDesc.Sizes[i + 1 - Offset];
 			}
 		}
 	}
 
-	void ComputeTensorStridesForBroadcast(const FMLTensorDesc& TensorDesc, uint32 OutputNumdimension, uint32 OutStrides[8])
+	void FillTensorStrideForBroadcastShaderParameters(const FMLTensorDesc& TensorDesc, uint32 OutputNumdimension, TStaticArray<FUintVector4, NXRT_TENSORSTRIDEINFO_MAX_NUM_DIMENSIONS, 16U>& OutShaderParam, uint32 Idx)
 	{
 		checkf(OutputNumdimension >= TensorDesc.Dimension, TEXT("Can't broadcast tensor from rank %d to rank %d, should be inferior or equal."), TensorDesc.Dimension, OutputNumdimension);
-		ComputeTensorStrides(TensorDesc, OutStrides, OutputNumdimension);
+		FillTensorStrideShaderParameters(TensorDesc, OutShaderParam, Idx, OutputNumdimension);
 		uint32 Offset = OutputNumdimension - TensorDesc.Dimension;
 		for (uint32 i = Offset; i < OutputNumdimension; ++i)
 		{
 			// the stride for broadcast dimension is kept as 0
 			if (TensorDesc.Sizes[i - Offset] == 1)
 			{
-				OutStrides[i] = 0;
+				OutShaderParam[i][Idx] = 0;
 			}
 		}
 	}
 
-	void FillTensorStrideShaderParameters(uint32 Strides[8], FUint32Vector4& OutShaderParam0, FUint32Vector4& OutShaderParam1)
-	{
-		for (uint32 i = 0; i < 4; ++i)
-		{
-			OutShaderParam0[i] = Strides[i];
-			OutShaderParam1[i] = Strides[i + 4];
-		}
-	}
-
-	void FillTensorStrideShaderParameters(const FMLTensorDesc& TensorDesc, FUint32Vector4& OutShaderParam0, FUint32Vector4& OutShaderParam1)
-	{
-		uint32 Strides[8];
-		ComputeTensorStrides(TensorDesc, Strides);
-		FillTensorStrideShaderParameters(Strides, OutShaderParam0, OutShaderParam1);
-	}
-
-	void FillTensorStrideForBroadcastShaderParameters(const FMLTensorDesc& TensorDesc, uint32 OutputNumdimension, FUint32Vector4& OutShaderParam0, FUint32Vector4& OutShaderParam1)
-	{
-		uint32 Strides[8];
-		ComputeTensorStridesForBroadcast(TensorDesc, OutputNumdimension, Strides);
-		FillTensorStrideShaderParameters(Strides, OutShaderParam0, OutShaderParam1);
-	}
-
 	FIntVector ComputeElementWiseThreadGroups(uint32 ElementCount, uint32 GroupSizeX)
 	{
 		FIntVector ThreadGroupCount;
diff --git a/Engine/Plugins/Experimental/NNX/Source/NNXRuntimeRDG/Private/NNXRuntimeHLSLHelper.h b/Engine/Plugins/Experimental/NNX/Source/NNXRuntimeRDG/Private/NNXRuntimeHLSLHelper.h
--- a/Engine/Plugins/Experimental/NNX/Source/NNXRuntimeRDG/Private/NNXRuntimeHLSLHelper.h
+++ b/Engine/Plugins/Experimental/NNX/Source/NNXRuntimeRDG/Private/NNXRuntimeHLSLHelper.h
@@ -3,14 +3,12 @@
 #pragma once
 
 #include "NNXTypes.h"
+#include "Containers/StaticArray.h"
+#include "NNXShaderParameters.h"
 
 namespace NNX
 {
-	void ComputeTensorStrides(const FMLTensorDesc& TensorDesc, uint32 OutStrides[8], uint32 TargetNumdimensionForBroadcast = -1);
-	void ComputeTensorStridesForBroadcast(const FMLTensorDesc& TensorDesc, uint32 OutputNumdimension, uint32 OutStrides[8]);
-	void FillTensorStrideShaderParameters(uint32 Strides[8], FUint32Vector4& OutShaderParam0, FUint32Vector4& OutShaderParam1);
-	void FillTensorStrideShaderParameters(const FMLTensorDesc& TensorDesc, FUint32Vector4& OutShaderParam0, FUint32Vector4& OutShaderParam1);
-	void FillTensorStrideForBroadcastShaderParameters(const FMLTensorDesc& TensorDesc, uint32 OutputNumdimension, FUint32Vector4& OutShaderParam0, FUint32Vector4& OutShaderParam1);
+	void FillTensorStrideShaderParameters(const FMLTensorDesc& TensorDesc, TStaticArray<FUintVector4, NXRT_TENSORSTRIDEINFO_MAX_NUM_DIMENSIONS, 16U>& OutShaderParam, uint32 Idx, uint32 TargetNumdimensionForBroadcast = -1);
+	void FillTensorStrideForBroadcastShaderParameters(const FMLTensorDesc& TensorDesc, uint32 OutputNumdimension, TStaticArray<FUintVector4, NXRT_TENSORSTRIDEINFO_MAX_NUM_DIMENSIONS, 16U>& OutShaderParam, uint32 Idx);
 	FIntVector ComputeElementWiseThreadGroups(uint32 ElementCount, uint32 GroupSizeX);
-
 } // NNX
\ No newline at end of file
