NNXQA Adjust test desc for mosaic model + few ops

Problem:
* Test were failing on platform

Solutions:
* mosaic precision was setuped for the ref runtime, should be for the one tested for
* adjust precision of trigo test for platforms
* skip matmul test with rank > 4 for NNXRuntimeDMLOrt (metacommand only support 2 to 4D)
* add onnx models to packaged test project
* add RDG insight to tes project (unrelated just nice for latter)

Tests:
* runned all test on platform

#rnx

diff --git a/Engine/Plugins/Experimental/NNX/Source/NNXQA/Private/NNXQAJsonUtils.h b/Engine/Plugins/Experimental/NNX/Source/NNXQA/Private/NNXQAJsonUtils.h
--- a/Engine/Plugins/Experimental/NNX/Source/NNXQA/Private/NNXQAJsonUtils.h
+++ b/Engine/Plugins/Experimental/NNX/Source/NNXQA/Private/NNXQAJsonUtils.h
@@ -113,10 +113,12 @@ namespace Json
 	{
 		FString Name;
 		TArray<FTestConfigDataset> Datasets;
+		TArray<FTestConfigRuntime> Runtimes;
 
 		BEGIN_JSON_SERIALIZER
 			JSON_SERIALIZE("name", Name);
 			JSON_SERIALIZE_ARRAY_SERIALIZABLE("datasets", Datasets, FTestConfigDataset);
+			JSON_SERIALIZE_ARRAY_SERIALIZABLE("runtimes", Runtimes, FTestConfigRuntime);
 		END_JSON_SERIALIZER
 	};
 
diff --git a/Engine/Plugins/Experimental/NNX/Source/NNXQA/Private/NNXQAParametricTest.cpp b/Engine/Plugins/Experimental/NNX/Source/NNXQA/Private/NNXQAParametricTest.cpp
--- a/Engine/Plugins/Experimental/NNX/Source/NNXQA/Private/NNXQAParametricTest.cpp
+++ b/Engine/Plugins/Experimental/NNX/Source/NNXQA/Private/NNXQAParametricTest.cpp
@@ -255,6 +255,7 @@ namespace Test
 								FTests::FTestSetup& Test = AddTest(TestCategoryPath, TestBaseName, TEXT(".") + GetTestSuffix(Dataset));
 								
 								ApplyRuntimesConfig(Test, TestCategory.Runtimes);
+								ApplyRuntimesConfig(Test, InputOutputSet.Runtimes);
 								ApplyTargetConfig(Test, TestTarget);
 								ApplyDatasetConfig(Test, Dataset, InputTypeFromTarget, OutputTypeFromTarget);
 								Test.IsModelTest = bIsModelCategory;
@@ -403,7 +404,7 @@ namespace Test
 				double EndTime = FPlatformTime::Seconds();
 				double TimeForTest = static_cast<float>(EndTime - StartTime);
 				if (bTestSucceeded)
-				{
+				{ 
 					UE_LOG(LogNNX, Display, TEXT("Tests succeeded in %0.1f seconds."), TimeForTest);
 				}
 				else
diff --git a/Engine/Plugins/Experimental/NNX/Source/NNXQA/Resources/NNXQATestDesc.json b/Engine/Plugins/Experimental/NNX/Source/NNXQA/Resources/NNXQATestDesc.json
--- a/Engine/Plugins/Experimental/NNX/Source/NNXQA/Resources/NNXQATestDesc.json
+++ b/Engine/Plugins/Experimental/NNX/Source/NNXQA/Resources/NNXQATestDesc.json
@@ -21,7 +21,7 @@
 						{ "name": "NNXRuntimeORTDml", "absolute_error": 5e-2, "relative_error": 50.0,
 							"comment": "Worst relative error 43.367836 % (epsilon 0.001000 %) at position 112558, got -0.003391 expected -0.002365. TODO we need an automatic balance between relative error and absolute error"
 						},
-						{ "name": "NNXRuntimeCPU", "absolute_error": 1e-3, "relative_error": 3.0 },
+						{ "name": "NNXRuntimeORTCpu", "absolute_error": 1e-3, "relative_error": 3.0 },
 						{ "name": "NNXRuntimeHlsl",	"skip": true, "comment": "Model with multiple layer not yet implemented."	},
 						{ "name": "NNXRuntimeDml",  "skip": true, "comment": "Model with multiple layer not yet implemented." }
 					]
@@ -72,13 +72,13 @@
 				{	"target": "Acos", "tags": [ "SmokeTest" ], "runtimes": [
 						{ "name": "NNXRuntimeORTDml", "absolute_error":  1e-4, "relative_error": 0.03},
 						{ "name": "NNXRuntimeDml",    "absolute_error":  1e-4, "relative_error": 0.03},
-						{ "name": "NNXRuntimeHlsl",   "absolute_error":  1e-4, "relative_error": 0.03}
+						{ "name": "NNXRuntimeHlsl",   "absolute_error":  2e-4, "relative_error": 0.03}
 					]},
 				{	"target": "Acosh" },
 				{	"target": "Asin", "runtimes": [
-						{ "name": "NNXRuntimeORTDml", "absolute_error":  1e-4, "relative_error": 1.5},
-						{ "name": "NNXRuntimeDml",    "absolute_error":  1e-4, "relative_error": 1.5},
-						{ "name": "NNXRuntimeHlsl",   "absolute_error":  1e-4, "relative_error": 1.5}
+						{ "name": "NNXRuntimeORTDml", "absolute_error":  1e-4, "relative_error": 2.5},
+						{ "name": "NNXRuntimeDml",    "absolute_error":  1e-4, "relative_error": 2.5},
+						{ "name": "NNXRuntimeHlsl",   "absolute_error":  2e-4, "relative_error": 2.5}
 					]},
 				{	"target": "Asinh", "runtimes": [
 						{ "name": "NNXRuntimeORTDml", "absolute_error":  1e-4, "relative_error": 0.03},
@@ -88,7 +88,7 @@
 				{	"target": "Atan", "runtimes": [
 						{ "name": "NNXRuntimeORTDml", "absolute_error":  1e-4, "relative_error": 0.03},
 						{ "name": "NNXRuntimeDml",    "absolute_error":  1e-4, "relative_error": 0.03},
-						{ "name": "NNXRuntimeHlsl",   "absolute_error":  1e-4, "relative_error": 0.03}
+						{ "name": "NNXRuntimeHlsl",   "absolute_error":  1.5e-3, "relative_error": 0.3}
 					]},
 				{	"target": "Atanh" },
 				{	"target": "BitShift", "skip": true, "comment": "TODO need attributes" },
@@ -378,7 +378,7 @@
 			]
 		},
 		{
-			"name": "MatMul",
+			"name": "MatMul.Regular",
 			"datasets": [
 				{ "inputs": [ { "shape": [ 2, 2 ] }, { "shape": [ 2, 2 ] } ], "outputs": [ { "shape": [ 2, 2 ] } ]},
 				{ "inputs": [ { "shape": [ 3, 4 ] }, { "shape": [ 4, 5 ] } ], "outputs": [ { "shape": [ 3, 5 ] } ]},
@@ -389,7 +389,16 @@
 				{ "inputs": [ { "shape": [ 1, 4 ] }, { "shape": [ 2, 4, 3 ] } ], "outputs": [ { "shape": [ 2, 1, 3 ] } ]},
 				{ "inputs": [ { "shape": [ 3, 4 ] }, { "shape": [ 4, 1 ] } ], "outputs": [ { "shape": [ 3, 1 ] } ]},
 				{ "inputs": [ { "shape": [ 2, 2, 3, 4 ] }, { "shape": [ 4, 1 ] } ], "outputs": [ { "shape": [ 2, 2, 3, 1 ] } ]},
-				{ "inputs": [ { "shape": [ 1, 6 ] }, { "shape": [ 6, 1 ] } ], "outputs": [ { "shape": [ 1, 1 ] } ]},
+				{ "inputs": [ { "shape": [ 1, 6 ] }, { "shape": [ 6, 1 ] } ], "outputs": [ { "shape": [ 1, 1 ] } ]}
+			]
+		},
+		{
+			"name": "MatMul.Rank5",
+			"runtimes": [	
+				{	"name": "NNXRuntimeORTDml", "skip": true, "comment": "DML metacommand only support 2 to 4 for inputs rank" 	},
+				{ "name": "NNXRuntimeHlsl", "relative_error": 0.01}
+			],
+			"datasets": [
 				{ "inputs": [ { "shape": [ 2, 2, 3, 5, 4 ] }, { "shape": [ 2, 1, 3, 4, 3 ] } ], "outputs": [ { "shape": [ 2, 2, 3, 5, 3 ] } ]},
 				{ "inputs": [ { "shape": [ 2, 1, 2, 5, 4 ] }, { "shape": [ 1, 3, 2, 4, 3 ] } ], "outputs": [ { "shape": [ 2, 3, 2, 5, 3 ] } ]},
 				{ "inputs": [ { "shape": [ 1, 1, 2, 5, 4 ] }, { "shape": [ 2, 3, 2, 4, 3 ] } ], "outputs": [ { "shape": [ 2, 3, 2, 5, 3 ] } ]},
diff --git a/Engine/Plugins/Experimental/NNX/Source/NNXQA/Resources/NNXQATestDescSchema.json b/Engine/Plugins/Experimental/NNX/Source/NNXQA/Resources/NNXQATestDescSchema.json
--- a/Engine/Plugins/Experimental/NNX/Source/NNXQA/Resources/NNXQATestDescSchema.json
+++ b/Engine/Plugins/Experimental/NNX/Source/NNXQA/Resources/NNXQATestDescSchema.json
@@ -219,6 +219,11 @@
 					"description": "Array of inputs/output to test",
 					"items": { "$ref": "#/$defs/test_config_dataset" }
 				},
+				"runtimes": {
+					"type": "array",
+					"description": "Array of runtimes to specialize test behavior for",
+					"items": { "$ref": "#/$defs/test_config_runtime" }
+				},
 				"comment": {
 					"type": "string",
 					"description": "Unused by test system, usefull to add clarification in json file"
