NNX Pow safe impl + test for nan in NNXQA

Problem:
* Pow test was failing on platform

Solutions:
* it appeared test code was not accounting for NaNs test should have failed for other HLSL platform too -> added NaN check in tensor comparison test
* added custom pow implementation matching the CPU behavior (handling negative pow for integer exponent)

Tests:
* runned all tests on platform and editor

#rnx

diff --git a/Engine/Plugins/Experimental/NNX/Source/NNXHlslShaders/Private/NNXElementWiseBinaryCS.cpp b/Engine/Plugins/Experimental/NNX/Source/NNXHlslShaders/Private/NNXElementWiseBinaryCS.cpp
--- a/Engine/Plugins/Experimental/NNX/Source/NNXHlslShaders/Private/NNXElementWiseBinaryCS.cpp
+++ b/Engine/Plugins/Experimental/NNX/Source/NNXHlslShaders/Private/NNXElementWiseBinaryCS.cpp
@@ -50,7 +50,7 @@ const FString FMLElementWiseBinaryCS::GetOpFunc(EMLElementWiseBinaryOperatorType
 	OP(Mul,              TEXT("((X)*(Y))"));
 	//OP(Or,             TEXT("((X)||(Y))"));
 	OP(Prelu,            TEXT("prelu(X,Y)"));
-	OP(Pow,              TEXT("pow(X,Y)"));
+	OP(Pow,              TEXT("safe_pow(X,Y)"));
 	OP(Sub,              TEXT("((X)-(Y))"));
 	//OP(Or,             TEXT("((X)^=(Y))"));
 
diff --git a/Engine/Plugins/Experimental/NNX/Source/NNXHlslShaders/Shaders/ElementWiseBinaryOp.usf b/Engine/Plugins/Experimental/NNX/Source/NNXHlslShaders/Shaders/ElementWiseBinaryOp.usf
--- a/Engine/Plugins/Experimental/NNX/Source/NNXHlslShaders/Shaders/ElementWiseBinaryOp.usf
+++ b/Engine/Plugins/Experimental/NNX/Source/NNXHlslShaders/Shaders/ElementWiseBinaryOp.usf
@@ -16,6 +16,25 @@ float prelu(float x, float slope)
 	return (x < 0.0f) ? (slope * x) : x;
 }
 
+float safe_pow(float base, float xp)
+{
+	if (base >= 0.0f)
+		return pow(base, xp);
+	
+	//negative base + not integer exponent => nan
+	if (round(xp) != xp)
+	{
+		
+		const float nan = 0. / 0.;
+		return nan;
+	}
+
+	//negative base + integer exponent
+	int isExpOdd = (xp % 2);
+	float sign = isExpOdd ? -1.0f : 1.0f;
+	return sign * pow(-base, xp);
+}
+
 [numthreads(THREADGROUP_SIZE_X, 1, 1)]
 void ElementWiseBinaryOp(in const uint3 DispatchThreadID : SV_DispatchThreadID)
 {
diff --git a/Engine/Plugins/Experimental/NNX/Source/NNXQA/Private/NNXQAUtils.cpp b/Engine/Plugins/Experimental/NNX/Source/NNXQA/Private/NNXQAUtils.cpp
--- a/Engine/Plugins/Experimental/NNX/Source/NNXQA/Private/NNXQAUtils.cpp
+++ b/Engine/Plugins/Experimental/NNX/Source/NNXQA/Private/NNXQAUtils.cpp
@@ -165,12 +165,42 @@ namespace Test
 		float WorstRelativeErrorRef = 0.0f;
 		float WorstRelativeErrorOther = 0.0f;
 
+		int32 NumExtraNaNsInResults = 0;
+		int32 FirstExtraNaNIndex = -1;
+		int32 NumMissingNaNsInResults = 0;
+		int32 FirstMissingNaNIndex = -1;
+
 		for (int32 i = 0; i < Volume; ++i)
 		{
 			//All type are converted to float for comparison purpose
 			const float Result = (float)OtherBuffer[i];
 			const float Reference = (float)RefBuffer[i];
 
+			
+			if (FMath::IsNaN(Result) && !FMath::IsNaN(Reference))
+			{
+				bTensorMemMatch = false;
+				++NumExtraNaNsInResults;
+				if (FirstExtraNaNIndex == -1)
+				{
+					FirstExtraNaNIndex = i;
+				}
+			}
+			if (!FMath::IsNaN(Result) && FMath::IsNaN(Reference))
+			{
+				bTensorMemMatch = false;
+				++NumMissingNaNsInResults;
+				if (FirstMissingNaNIndex == -1)
+				{
+					FirstMissingNaNIndex = i;
+				}
+			}
+			
+			if (FMath::IsNaN(Result) || FMath::IsNaN(Reference))
+			{
+				continue;
+			}
+
 			const float AbsoluteError = FMath::Abs<float>(Result - Reference);
 			const float RelativeError = 100.0f * (AbsoluteError / FMath::Abs<float>(Reference));
 			if (AbsoluteError > AbsoluteErrorEpsilon || RelativeError > RelativeErrorPercent)
@@ -195,9 +225,13 @@ namespace Test
 
 		if (!bTensorMemMatch)
 		{
-			UE_LOG(LogNNX, Error, TEXT("Tensor data do not match.\nLogNNX: Worst absolute error %f (epsilon %f) at position %d, got %f expected %f\nLogNNX: Worst relative error %f%% (epsilon %f%%) at position %d, got %f expected %f"),
+			UE_LOG(LogNNX, Error, TEXT("Tensor data do not match.\n"
+				"LogNNX: Worst absolute error %f (epsilon %f) at position %d, got %f expected %f\n"
+				"LogNNX: Worst relative error % f % %(epsilon % f%%) at position % d, got % f expected % f\n"
+				"LogNNX: Num unexpected NaNs %d (First at index %d), Num missing NaNs %d (first at index %d)"),
 				WorstAbsoluteError, AbsoluteErrorEpsilon, WorstAbsoluteErrorIndex, WorstAbsoluteErrorOther, WorstAbsoluteErrorRef,
-				WorstRelativeError, RelativeErrorPercent, WorstRelativeErrorIndex, WorstRelativeErrorOther, WorstRelativeErrorRef);
+				WorstRelativeError, RelativeErrorPercent, WorstRelativeErrorIndex, WorstRelativeErrorOther, WorstRelativeErrorRef,
+				NumExtraNaNsInResults, FirstExtraNaNIndex, NumMissingNaNsInResults, FirstMissingNaNIndex);
 			UE_LOG(LogNNX, Error, TEXT("   Expected : %s"), *TensorToString(RefTensorDesc, RefRawBuffer));
 			UE_LOG(LogNNX, Error, TEXT("   But got  : %s"), *TensorToString(OtherTensorDesc, OtherRawBuffer));
 			return false;
