Don't register Ort DML runtime if Project Settings are not set to DirectX 12

#rnx

diff --git a/Engine/Plugins/Experimental/NNX/Source/NNXRuntimeORT/Private/NNXRuntimeORT.cpp b/Engine/Plugins/Experimental/NNX/Source/NNXRuntimeORT/Private/NNXRuntimeORT.cpp
--- a/Engine/Plugins/Experimental/NNX/Source/NNXRuntimeORT/Private/NNXRuntimeORT.cpp
+++ b/Engine/Plugins/Experimental/NNX/Source/NNXRuntimeORT/Private/NNXRuntimeORT.cpp
@@ -1,6 +1,9 @@
 // Copyright Epic Games, Inc. All Rights Reserved.
 #include "NNXRuntimeORT.h"
 
+#include "CoreGlobals.h"
+#include "Misc/ConfigCacheIni.h"
+
 #include "NeuralTimer.h"
 #include "NNXRuntimeORTUtils.h"
 #include "RedirectCoutAndCerrToUeLog.h"
@@ -455,4 +458,38 @@ bool FMLInferenceModelORTDml::InitializedAndConfigureMembers()
 
 	return true;
 }
+
+IRuntime* NNX::FRuntimeORTDMLStartup()
+{
+	if (!GORTDMLRuntime)
+	{
+		// In order to register DirectML we need D3D12
+		bool	bHasD3D12Config = false;
+		bool	bHasD3D12RHI = false;
+		FString DefaultGraphicsRHI;
+
+		if (GConfig->GetString(TEXT("/Script/WindowsTargetPlatform.WindowsTargetSettings"), TEXT("DefaultGraphicsRHI"), DefaultGraphicsRHI, GEngineIni))
+		{
+			bHasD3D12Config = DefaultGraphicsRHI == TEXT("DefaultGraphicsRHI_DX12");
+		}
+
+		// We need to check if RHI is forced to D3D12
+		if (GDynamicRHI && GDynamicRHI->GetInterfaceType() == ERHIInterfaceType::D3D12)
+		{
+			bHasD3D12RHI = true;
+		}
+
+		if (bHasD3D12Config && bHasD3D12RHI)
+		{
+			GORTDMLRuntime = FRuntimeORTDMLCreate();
+		}
+		else
+		{
+			return nullptr;
+		}
+	}
+
+	return GORTDMLRuntime.Get();
+};
+
 #endif
diff --git a/Engine/Plugins/Experimental/NNX/Source/NNXRuntimeORT/Private/NNXRuntimeORT.h b/Engine/Plugins/Experimental/NNX/Source/NNXRuntimeORT/Private/NNXRuntimeORT.h
--- a/Engine/Plugins/Experimental/NNX/Source/NNXRuntimeORT/Private/NNXRuntimeORT.h
+++ b/Engine/Plugins/Experimental/NNX/Source/NNXRuntimeORT/Private/NNXRuntimeORT.h
@@ -140,14 +140,7 @@ namespace NNX
 		return GORTCUDARuntime.Get();
 	};
 
-	inline IRuntime* FRuntimeORTDMLStartup()
-	{
-		if (!GORTDMLRuntime)
-		{
-			GORTDMLRuntime = FRuntimeORTDMLCreate();
-		}
-		return GORTDMLRuntime.Get();
-	};
+	extern IRuntime* FRuntimeORTDMLStartup();
 #endif
 
 	inline void FRuntimeORTCPUShutdown()
