Refactor readback code

#rnx

diff --git a/Engine/Plugins/Experimental/NNX/Source/NNXRuntimeRDG/Private/NNXRuntimeDML.cpp b/Engine/Plugins/Experimental/NNX/Source/NNXRuntimeRDG/Private/NNXRuntimeDML.cpp
--- a/Engine/Plugins/Experimental/NNX/Source/NNXRuntimeRDG/Private/NNXRuntimeDML.cpp
+++ b/Engine/Plugins/Experimental/NNX/Source/NNXRuntimeRDG/Private/NNXRuntimeDML.cpp
@@ -1068,6 +1068,7 @@ bool FMLRuntimeDml::RegisterElementWiseBinaryOperators()
 //
 FMLInferenceModelDml::FMLInferenceModelDml()
 {
+	bUseManualTransitions = true;
 }
 
 //
diff --git a/Engine/Plugins/Experimental/NNX/Source/NNXRuntimeRDG/Private/NNXRuntimeRDG.cpp b/Engine/Plugins/Experimental/NNX/Source/NNXRuntimeRDG/Private/NNXRuntimeRDG.cpp
--- a/Engine/Plugins/Experimental/NNX/Source/NNXRuntimeRDG/Private/NNXRuntimeRDG.cpp
+++ b/Engine/Plugins/Experimental/NNX/Source/NNXRuntimeRDG/Private/NNXRuntimeRDG.cpp
@@ -17,7 +17,14 @@ namespace NNX
 //
 FMLInferenceModelRDG::FMLInferenceModelRDG()
 	: FMLInferenceModel(EMLInferenceModelType::RDG)
+	, bUseManualTransitions(false)
 {
+	Readback.RHI = new FRHIGPUBufferReadback("FMLTensorReadback");
+}
+
+FMLInferenceModelRDG::~FMLInferenceModelRDG()
+{
+	delete Readback.RHI;
 }
 
 //
@@ -102,6 +109,24 @@ int FMLInferenceModelRDG::Run(TArrayView<const FMLTensorBinding> InInputBindings
 			if (Res == 0)
 			{
 				GraphBuilder.Execute();
+
+				// FIXME: Using BlockUntilGPUIdle() prevents hang on Linux
+				RHICmdList.BlockUntilGPUIdle();
+				//RHICmdList.SubmitCommandsHint();
+
+				// Wait for readback
+				while (!Readback.RHI->IsReady())
+				{
+					FPlatformProcess::Sleep(0.001f);
+				}
+
+				// Process readback
+				{
+					const void* BuffData = Readback.RHI->Lock(Readback.Size);
+					check(BuffData);
+					FMemory::Memcpy(Readback.CpuMemory, BuffData, Readback.Size);
+					Readback.RHI->Unlock();
+				}
 			}
 			
 			Signal->Trigger();
@@ -111,6 +136,8 @@ int FMLInferenceModelRDG::Run(TArrayView<const FMLTensorBinding> InInputBindings
 	// We need to wait for render thread to finish
 	Signal->Wait();
 
+	FGenericPlatformProcess::ReturnSynchEventToPool(Signal);
+
 	return Res;
 }
 
@@ -157,7 +184,7 @@ int FMLInferenceModelRDG::EnqueueRDG(FRDGBuilder& GraphBuilder, TArrayView<const
 	// If required, readback the output tensors to CPU
 	if (!RDGReadbackIndices.IsEmpty())
 	{
-		AddTensorReadbacks_RenderThread(GraphBuilder, RDGReadbackIndices, RDGOutputBindings, InOutputBindings);
+		AddTensorReadbacks_RenderThread(GraphBuilder, RDGReadbackIndices, RDGOutputBindings, InOutputBindings);		
 	}
 
 	return 0;
@@ -176,15 +203,16 @@ int FMLInferenceModelRDG::SetTensors(FRDGBuilder& GraphBuilder, FMLTensorBinding
 
 		if (Binding.BindingType == EMLTensorBindingDataType::CPUMemory)
 		{
-			FRDGBufferRef	TensorBuffer;
+			// FIXME: CreateStructuredDesc() creates a crash on VulkanRHI
+			//FRDGBufferDesc Desc = FRDGBufferDesc::CreateStructuredDesc(TensorDesc.GetElemByteSize(), TensorDesc.Num());
+			FRDGBufferDesc Desc = FRDGBufferDesc::CreateBufferDesc(TensorDesc.GetElemByteSize(), TensorDesc.Num());
 
-			TensorBuffer = GraphBuilder.CreateBuffer(
-				FRDGBufferDesc::CreateBufferDesc(TensorDesc.GetElemByteSize(), TensorDesc.Num()),
-				*TensorDesc.Name,
-				ERDGBufferFlags::None);
+			// FIXME: We should use BUF_SourceCopy for only output buffers (GPU readback)
+			Desc.Usage = EBufferUsageFlags(Desc.Usage | BUF_SourceCopy);
 
+			FRDGBufferRef TensorBuffer = GraphBuilder.CreateBuffer(Desc, *TensorDesc.Name, ERDGBufferFlags::None);
+			
 			OutBindings.Emplace(FMLTensorBinding::FromRDG(TensorBuffer, InTensors[Idx].DataSize));
-
 			OutIndices.Add(Idx);
 		}
 		else if (Binding.BindingType == EMLTensorBindingDataType::RDGBuffer)
@@ -213,21 +241,7 @@ void FMLInferenceModelRDG::AddTensorUploads_RenderThread(FRDGBuilder& GraphBuild
 		const FMLTensorBinding& InBinding = InBindings[TensorIdx];
 		const FMLTensorDesc&	TensorDesc = InputTensors[TensorIdx];
 
-		FMLTensorUploadParameters* TensorUploadParams = GraphBuilder.AllocParameters<FMLTensorUploadParameters>();
-
-		TensorUploadParams->Buffer = RDGBinding.Buffer;
-
-		GraphBuilder.AddPass(
-			RDG_EVENT_NAME("FMLInferenceModelAddTensorUpload"),
-			TensorUploadParams,
-			ERDGPassFlags::Copy | ERDGPassFlags::NeverCull,
-			[InBinding, TensorDesc, TensorUploadParams](FRHICommandListImmediate& RHICmdList)
-			{
-				void* BuffData = RHICmdList.LockBuffer(TensorUploadParams->Buffer->GetRHI(), 0, TensorDesc.DataSize, RLM_WriteOnly);
-				FMemory::Memcpy(BuffData, InBinding.CpuMemory, TensorDesc.DataSize);
-				RHICmdList.UnlockBuffer(TensorUploadParams->Buffer->GetRHI());
-			}
-		);
+		GraphBuilder.QueueBufferUpload(RDGBinding.Buffer, InBinding.CpuMemory, TensorDesc.DataSize, ERDGInitialDataFlags::NoCopy);
 	}
 }
 
@@ -240,7 +254,7 @@ void FMLInferenceModelRDG::AddTensorReadbacks_RenderThread(FRDGBuilder& GraphBui
 	{
 		const int32				TensorIdx = InReadbackIndices[Idx];
 		const FMLTensorBinding& RDGBinding = InRDGBindings[TensorIdx];
-		const FMLTensorBinding& InBinding = InBindings[TensorIdx];
+		const FMLTensorBinding& Binding = InBindings[TensorIdx];
 		const FMLTensorDesc&	TensorDesc = OutputTensors[TensorIdx];
 
 		FMLTensorReadbackParameters* TensorReadbackParams = GraphBuilder.AllocParameters<FMLTensorReadbackParameters>();
@@ -248,43 +262,29 @@ void FMLInferenceModelRDG::AddTensorReadbacks_RenderThread(FRDGBuilder& GraphBui
 		TensorReadbackParams->Buffer = RDGBinding.Buffer;
 
 		GraphBuilder.AddPass(
-			RDG_EVENT_NAME("FMLInferenceModelAddTensorReadback"),
+			RDG_EVENT_NAME("FMLInferenceModelAddTensorReadback:%s", *TensorDesc.Name),
 			TensorReadbackParams,
 			ERDGPassFlags::Readback | ERDGPassFlags::NeverCull,
-			[InBinding, TensorDesc, TensorReadbackParams](FRHICommandListImmediate& RHICmdList)
+			[this, Binding, TensorDesc, TensorReadbackParams](FRHICommandListImmediate& RHICmdList)
 			{
 				FRHIBuffer* OutputBuffer = TensorReadbackParams->Buffer->GetRHI();
 
-				FRHITransitionInfo PreTransitions[] =
-				{
-					FRHITransitionInfo(OutputBuffer, ERHIAccess::UAVCompute, ERHIAccess::CopySrc)
-				};
-
-				RHICmdList.Transition(MakeArrayView(PreTransitions, UE_ARRAY_COUNT(PreTransitions)));
-
 				// TODO: FIXME: We need to transition the resources for DirectML
-				RHICmdList.SubmitCommandsHint();
-
-				FRHIGPUBufferReadback	Readback("FMLTensorReadback");
-
-				Readback.EnqueueCopy(RHICmdList, OutputBuffer, TensorDesc.DataSize);
-				RHICmdList.BlockUntilGPUIdle();
-				check(Readback.IsReady());
-
-				const void* BuffData = Readback.Lock(TensorDesc.DataSize);
-				FMemory::Memcpy(InBinding.CpuMemory, BuffData, TensorDesc.DataSize);
-				Readback.Unlock();
-
-				//const void* BuffData = RHICmdList.LockBuffer(OutputBuffer, 0, TensorDesc.DataSize, RLM_ReadOnly);
-				//FMemory::Memcpy(InBinding.CpuMemory, BuffData, TensorDesc.DataSize);
-				//RHICmdList.UnlockBuffer(OutputBuffer);
-
-				//FRHITransitionInfo PostTransitions[] =
-				//{
-				//	FRHITransitionInfo(OutputBuffer, ERHIAccess::CopySrc, ERHIAccess::UAVCompute)
-				//};
-
-				//RHICmdList.Transition(MakeArrayView(PostTransitions, UE_ARRAY_COUNT(PostTransitions)));
+				if (bUseManualTransitions)
+				{
+					FRHITransitionInfo Transitions[] =
+					{
+						FRHITransitionInfo(OutputBuffer, ERHIAccess::UAVCompute, ERHIAccess::CopySrc)
+					};
+
+					RHICmdList.Transition(MakeArrayView(Transitions, UE_ARRAY_COUNT(Transitions)));
+					RHICmdList.SubmitCommandsHint();
+				}
+
+				Readback.RHI->EnqueueCopy(RHICmdList, OutputBuffer, TensorDesc.DataSize);
+				Readback.CpuMemory = Binding.CpuMemory;
+				Readback.Offset = 0;
+				Readback.Size = TensorDesc.DataSize;
 			}
 		);
 	}
diff --git a/Engine/Plugins/Experimental/NNX/Source/NNXRuntimeRDG/Private/NNXRuntimeRDG.h b/Engine/Plugins/Experimental/NNX/Source/NNXRuntimeRDG/Private/NNXRuntimeRDG.h
--- a/Engine/Plugins/Experimental/NNX/Source/NNXRuntimeRDG/Private/NNXRuntimeRDG.h
+++ b/Engine/Plugins/Experimental/NNX/Source/NNXRuntimeRDG/Private/NNXRuntimeRDG.h
@@ -5,6 +5,7 @@
 #include "NNXCore.h"
 #include "NNXRuntime.h"
 #include "ShaderParameterUtils.h"
+#include "RHIGPUReadback.h"
 
 #include "Containers/Map.h"
 
@@ -44,8 +45,18 @@ using FMLIntArray = TArray<int32, TInlineAllocator<16>>;
  */
 class FMLInferenceModelRDG : public FMLInferenceModel
 {
+	struct FReadbackEntry
+	{
+		FRHIGPUBufferReadback*	RHI;
+		void*					CpuMemory;
+		size_t					Offset;
+		size_t					Size;
+	};
+
 public:
 
+	~FMLInferenceModelRDG();
+
 	virtual int Run(TArrayView<const FMLTensorBinding> InInputBindings, TArrayView<const FMLTensorBinding> OutOutputBindings) override;
 	virtual int EnqueueRDG(FRDGBuilder& Builder, TArrayView<const FMLTensorBinding> InInputBindings, TArrayView<const FMLTensorBinding> OutOutputBindings) override;
 
@@ -61,6 +72,9 @@ protected:
 
 	virtual void AddTensorUploads_RenderThread(FRDGBuilder& GraphBuilder, TArrayView<const int32> InUploadIndices, TArrayView<FMLTensorBinding> InRDGBindings, TArrayView<const FMLTensorBinding> InBindings);
 	virtual void AddTensorReadbacks_RenderThread(FRDGBuilder& GraphBuilder, TArrayView<const int32> InReadbackIndices, TArrayView<const FMLTensorBinding> InRDGBindings, TArrayView<const FMLTensorBinding> InBindings);
+
+	FReadbackEntry	Readback;
+	bool			bUseManualTransitions;
 };
 
 
