NNX operator Attribute support

Problems:
* Operator could not use attribute from the inference models on RDG
* We could not test behavior of operator using non default-attribute

Solutions:
* Add support for attribute serialization in RDG serializable format
* Add POC handling of attribute at operator initialisation for NNXRuntimeHLSL (DML TODO)
* Add attribute support to model creation helper for tests

Tests:
* all green test in editor
* added a hacked LeakyRelu test with attribute (will extent parametric tests as a follow up CL)

diff --git a/Engine/Plugins/Experimental/NNX/Source/NNXCore/Public/NNXRuntimeFormat.h b/Engine/Plugins/Experimental/NNX/Source/NNXCore/Public/NNXRuntimeFormat.h
--- a/Engine/Plugins/Experimental/NNX/Source/NNXCore/Public/NNXRuntimeFormat.h
+++ b/Engine/Plugins/Experimental/NNX/Source/NNXCore/Public/NNXRuntimeFormat.h
@@ -16,6 +16,19 @@ enum class EMLFormatTensorType : uint8
 	Intermediate
 };
 
+// Required by LoadModel() when loading operators in HLSL and DirectML runtime
+USTRUCT()
+struct FMLFormatAttributeDesc
+{
+	GENERATED_USTRUCT_BODY()
+
+	UPROPERTY(VisibleAnywhere, Category = "Neural Network Inference")
+	FString Name;
+
+	UPROPERTY(VisibleAnywhere, Category = "Neural Network Inference")
+	FMLAttributeValue Value;
+};
+
 USTRUCT()
 struct FMLFormatOperatorDesc
 {
@@ -30,6 +43,9 @@ struct FMLFormatOperatorDesc
 	UPROPERTY(VisibleAnywhere, Category = "Neural Network Inference")
 	TArray<uint32> OutTensors;
 
+	UPROPERTY(VisibleAnywhere, Category = "Neural Network Inference")
+	TArray<FMLFormatAttributeDesc> Attributes;
+
 	//UPROPERTY(VisibleAnywhere, Category = "Neural Network Inference")
 	//FString InstanceName;		//!< For example "Relu1"
 
@@ -127,4 +143,4 @@ struct FMLRuntimeFormat
 
 	UPROPERTY(VisibleAnywhere, Category = "Neural Network Inference")
 	TArray<uint8> TensorData;
-};
+};
\ No newline at end of file
diff --git a/Engine/Plugins/Experimental/NNX/Source/NNXCore/Public/NNXTypes.h b/Engine/Plugins/Experimental/NNX/Source/NNXCore/Public/NNXTypes.h
--- a/Engine/Plugins/Experimental/NNX/Source/NNXCore/Public/NNXTypes.h
+++ b/Engine/Plugins/Experimental/NNX/Source/NNXCore/Public/NNXTypes.h
@@ -3,6 +3,8 @@
 #pragma once
 
 #include "UObject/Object.h"
+#include "UObject/Class.h"
+#include "NNXTypes.generated.h"
 
 /** Tensor data types */
 UENUM()
@@ -27,6 +29,148 @@ enum class EMLTensorDataType : uint8
 	BFloat16							//!< 16-bit floating number
 };
 
+/** Attribute data types */
+UENUM()
+enum class EMLAttributeDataType : uint8
+{
+	None,
+	Float,								//!< 32-bit floating number
+	Int32								//!< 32-bit signed integer
+    //TODO add more AttributeDataType support
+};
+
+USTRUCT()
+struct FMLAttributeValue
+{
+	GENERATED_USTRUCT_BODY()
+
+	FMLAttributeValue()
+	{
+		Type = EMLAttributeDataType::None;
+	}
+
+    //TODO add more AttributeDataType support
+	explicit FMLAttributeValue(float Value)
+	{
+		Type = EMLAttributeDataType::Float;
+		Data.SetNumUninitialized(sizeof(float));
+		FMemory::Memcpy(Data.GetData(), &Value, sizeof(float));
+	}
+
+	explicit FMLAttributeValue(int Value)
+	{
+		Type = EMLAttributeDataType::Int32;
+		Data.SetNumUninitialized(sizeof(int32));
+		FMemory::Memcpy(Data.GetData(), &Value, sizeof(int32));
+	}
+
+	float AsFloat() const
+	{
+		check(Type == EMLAttributeDataType::Float);
+		return *(float*)Data.GetData();
+	}
+
+	int32 AsInt32() const
+	{
+		check(Type == EMLAttributeDataType::Int32);
+		return *(int32*)Data.GetData();
+	}
+
+	EMLAttributeDataType GetType() const
+	{
+		return Type;
+	}
+	
+private:
+	
+	UPROPERTY(VisibleAnywhere, Category = "Neural Network Inference")
+	EMLAttributeDataType	Type;
+	
+	UPROPERTY(VisibleAnywhere, Category = "Neural Network Inference")
+	TArray<uint8>			Data;
+};
+
+class FMLAttributeMap
+{
+public:
+	//Set attribute
+	void SetAttribute(const FString& Name, const FMLAttributeValue& Value)
+	{
+		checkf(nullptr == Attributes.FindByPredicate([Name](const FEntry& Other) { return Other.Name == Name; }), TEXT("Attribute name should be unique"));
+		Attributes.Emplace(Name, Value);
+	}
+
+	//Query attributes
+	float GetFloat(const FString& Name) const
+	{
+		const FMLAttributeValue* Value = GetAttributeValue(Name);
+		checkf(Value != nullptr, TEXT("Required attribute %s not found"), *Name);
+		return Value->AsFloat();
+	}
+
+	int32 GetInt32(const FString& Name) const
+	{
+		const FMLAttributeValue* Value = GetAttributeValue(Name);
+		checkf(Value != nullptr, TEXT("Required attribute %s not found"), *Name);
+		return Value->AsInt32();
+	}
+
+	float GetOptionalFloat(const FString& Name, float Default) const
+	{
+		const FMLAttributeValue* Value = GetAttributeValue(Name);
+		return Value == nullptr ? Default : Value->AsFloat();
+	}
+
+	int32 GetOptionalInt32(const FString& Name, int32 Default) const
+	{
+		const FMLAttributeValue* Value = GetAttributeValue(Name);
+		return Value == nullptr ? Default : Value->AsInt32();
+	}
+	
+	const FMLAttributeValue* GetAttributeValue(const FString& Name) const
+	{
+		const FEntry * entry = Attributes.FindByPredicate([Name](const FEntry& Other) { return Other.Name == Name; });
+		if (entry != nullptr)
+		{
+			return &entry->Value;
+		}
+		return nullptr;
+	}
+
+	//Iteration
+	int32 Num() const
+	{
+		return Attributes.Num();
+	}
+
+	const FString& GetName(int32 Idx) const
+	{
+		check(Idx < Num());
+		return Attributes[Idx].Name;
+	}
+
+	const FMLAttributeValue& GetAttributeValue(int32 Idx) const
+	{
+		check(Idx < Num());
+		return Attributes[Idx].Value;
+	}
+
+private:
+
+	struct FEntry
+	{
+		FEntry(const FString& InName, const FMLAttributeValue& InValue)
+			: Name(InName), Value(InValue)
+		{
+		}
+		
+		FString Name;
+		FMLAttributeValue Value;
+	};
+
+	TArray<FEntry> Attributes;
+};
+
 namespace NNX
 {
 
diff --git a/Engine/Plugins/Experimental/NNX/Source/NNXQA/Private/NNXQAParametricTest.cpp b/Engine/Plugins/Experimental/NNX/Source/NNXQA/Private/NNXQAParametricTest.cpp
--- a/Engine/Plugins/Experimental/NNX/Source/NNXQA/Private/NNXQAParametricTest.cpp
+++ b/Engine/Plugins/Experimental/NNX/Source/NNXQA/Private/NNXQAParametricTest.cpp
@@ -303,8 +303,14 @@ namespace Test
 		}
 		else
 		{
+			//TODO add support for attribute in json schema
+			FMLAttributeMap AttributeMap;
+			if (TestSetup.TestName == TEXT("System.Engine.MachineLearning.NNX.Operator.Elementwise.Unary.LeakyRelu.[4]=>[4]"))
+			{
+				AttributeMap.SetAttribute("alpha", FMLAttributeValue(0.5f));
+			}
 			// Operator test, create model in memory
-			if (!CreateONNXModelForOperator(TestSetup.TargetName, TestSetup.Inputs, TestSetup.Outputs, ModelData))
+			if (!CreateONNXModelForOperator(TestSetup.TargetName, TestSetup.Inputs, TestSetup.Outputs, AttributeMap, ModelData))
 			{
 				UE_LOG(LogNNX, Error, TEXT("Failed to create model for test '%s'. Test ABORTED!"), *TestSetup.TargetName);
 				return false;
diff --git a/Engine/Plugins/Experimental/NNX/Source/NNXQA/Resources/NNXQATestDesc.json b/Engine/Plugins/Experimental/NNX/Source/NNXQA/Resources/NNXQATestDesc.json
--- a/Engine/Plugins/Experimental/NNX/Source/NNXQA/Resources/NNXQATestDesc.json
+++ b/Engine/Plugins/Experimental/NNX/Source/NNXQA/Resources/NNXQATestDesc.json
@@ -116,7 +116,7 @@
 				{	"target": "IsNan", "skip": true, "comment": "TODO add flexibility in test data setup + bool tensors" },
 				{	"target": "HardSigmoid",         "comment": "TODO test with not-default attribute" },
 				{	"target": "HardSwish", "runtimes": [ { "name": "NNXRuntimeDml", "skip": true, "comment": "Not implemented atm on NNXRuntimeDML."} ]},
-				{	"target": "LeakyRelu",           "comment": "TODO test with not-default attribute" },
+				{	"target": "LeakyRelu", "additional_datasets": [ "Elementwise.LeakyRelu.HackedAttributeTest" ], "comment": "TODO test with not-default attribute" },
 				{	"target": "Log", "runtimes": [
 						{ "name": "NNXRuntimeORTDml", "relative_error": 0.15},
 						{ "name": "NNXRuntimeDml",    "relative_error": 0.15},
@@ -274,6 +274,16 @@
 				}
 			]
 		},
+		{
+			"name": "Elementwise.LeakyRelu.HackedAttributeTest",
+			"datasets": [
+				{
+					"inputs": [ { "shape": [ 4 ] } ],
+					"runtimes": [ { "name": "NNXRuntimeDML",  "skip": true, "comment": "RDG DML does not support attribute yet" } ],
+					"comment": "for now there is no way to set attribute using the json test desc, this dataset is a special temporary case to test LeakyRelu with attribute not defaulted. See RunParametricTest() in NNQAParametricTest.cpp"
+				}
+			]
+		},
 		{
 			"name": "Elementwise.Binary.NoBroadcast",
 			"datasets": [
diff --git a/Engine/Plugins/Experimental/NNX/Source/NNXRuntimeRDG/Private/NNXRuntimeHLSL.cpp b/Engine/Plugins/Experimental/NNX/Source/NNXRuntimeRDG/Private/NNXRuntimeHLSL.cpp
--- a/Engine/Plugins/Experimental/NNX/Source/NNXRuntimeRDG/Private/NNXRuntimeHLSL.cpp
+++ b/Engine/Plugins/Experimental/NNX/Source/NNXRuntimeRDG/Private/NNXRuntimeHLSL.cpp
@@ -35,7 +35,7 @@ protected:
 
 private:
 
-	FMLOperatorHlsl* OpCreate(const FString& Name, TArrayView<const FMLTensorDesc> InputTensorDesc, TArrayView<const FMLTensorDesc> OutputTensorDescs);
+	FMLOperatorHlsl* OpCreate(const FString& Name, TArrayView<const FMLTensorDesc> InputTensorDesc, TArrayView<const FMLTensorDesc> OutputTensorDescs, const FMLAttributeMap& Attributes);
 
 	TArray<FMLOperatorHlsl*>	Operators;
 };
@@ -83,7 +83,14 @@ bool FMLInferenceModelHlsl::Init(UMLInferenceModel* InModel)
 		TArray<FMLTensorDesc> OpInputTensors = InputTensors;
 		TArray<FMLTensorDesc> OpOutputTensors = OutputTensors;
 
-		FMLOperatorHlsl* Op = OpCreate(TypeName, OpInputTensors, OpOutputTensors);
+		// Attributes
+		FMLAttributeMap AttributeMap;
+		for (const FMLFormatAttributeDesc& Desc : Format.Operators[Idx].Attributes)
+		{
+			AttributeMap.SetAttribute(Desc.Name, Desc.Value);
+		}
+
+		FMLOperatorHlsl* Op = OpCreate(TypeName, OpInputTensors, OpOutputTensors, AttributeMap);
 
 		if (!Op) //Op.Shader.IsNull())
 		{
@@ -102,7 +109,7 @@ bool FMLInferenceModelHlsl::Init(UMLInferenceModel* InModel)
 //
 //
 //
-FMLOperatorHlsl* FMLInferenceModelHlsl::OpCreate(const FString& OpName, TArrayView<const FMLTensorDesc> InputTensorDescs, TArrayView<const FMLTensorDesc> OutputTensorDescs)
+FMLOperatorHlsl* FMLInferenceModelHlsl::OpCreate(const FString& OpName, TArrayView<const FMLTensorDesc> InputTensorDescs, TArrayView<const FMLTensorDesc> OutputTensorDescs, const FMLAttributeMap& AttributeMap)
 {
 	FMLOperatorRegistryHlsl::MLOperatorCreateFunc CreateFn = FMLOperatorRegistryHlsl::Get()->OpFind(OpName);
 
@@ -114,7 +121,7 @@ FMLOperatorHlsl* FMLInferenceModelHlsl::OpCreate(const FString& OpName, TArrayVi
 
 	FMLOperatorHlsl* Op = CreateFn();
 
-	Op->Initialize(InputTensorDescs, OutputTensorDescs);
+	Op->Initialize(InputTensorDescs, OutputTensorDescs, AttributeMap);
 
 	return Op;
 }
diff --git a/Engine/Plugins/Experimental/NNX/Source/NNXRuntimeRDG/Private/NNXRuntimeHLSLConvTransposeOp.cpp b/Engine/Plugins/Experimental/NNX/Source/NNXRuntimeRDG/Private/NNXRuntimeHLSLConvTransposeOp.cpp
--- a/Engine/Plugins/Experimental/NNX/Source/NNXRuntimeRDG/Private/NNXRuntimeHLSLConvTransposeOp.cpp
+++ b/Engine/Plugins/Experimental/NNX/Source/NNXRuntimeRDG/Private/NNXRuntimeHLSLConvTransposeOp.cpp
@@ -6,6 +6,7 @@
 
 namespace NNX
 {
+    DECLARE_GPU_STAT_NAMED(FMLHLSLOperatorConvTranspose, TEXT("FML.HLSL.Operator.ConvTranspose"));
 
 	/**
 	 * ConvTranspose ML operator implementation
@@ -45,7 +46,7 @@ namespace NNX
 
 	public:
 
-		virtual bool Initialize(TArrayView<const FMLTensorDesc> InputTensors, TArrayView<const FMLTensorDesc> OutputTensors) override
+		virtual bool Initialize(TArrayView<const FMLTensorDesc> InputTensors, TArrayView<const FMLTensorDesc> OutputTensors, const FMLAttributeMap& Attributes) override
 		{
 			check(InputTensors.Num() >= 2 && InputTensors.Num() <= 3);
 			check(OutputTensors.Num() == 1);
@@ -107,9 +108,12 @@ namespace NNX
 			PermutationVector.Set<FMLConvTransposeCS::FConvTransposeHasB>(HasBias);
 			TShaderMapRef<FMLConvTransposeCS> ComputeShader(GetGlobalShaderMap(GMaxRHIFeatureLevel), PermutationVector);
 
+			RDG_EVENT_SCOPE(GraphBuilder, "FML.HLSL.Operator.ConvTranspose");
+			RDG_GPU_STAT_SCOPE(GraphBuilder, FMLHLSLOperatorConvTranspose);
+
 			FComputeShaderUtils::AddPass(
 				GraphBuilder,
-				RDG_EVENT_NAME("FMLOperatorHlslConvTranspose_Dispatch"),
+				RDG_EVENT_NAME("FML.HLSL.Operator.ConvTranspose.Dispatch"),
 				ERDGPassFlags::Compute | ERDGPassFlags::NeverCull,
 				ComputeShader,
 				Params,
diff --git a/Engine/Plugins/Experimental/NNX/Source/NNXRuntimeRDG/Private/NNXRuntimeHLSLElementWiseBinaryOps.cpp b/Engine/Plugins/Experimental/NNX/Source/NNXRuntimeRDG/Private/NNXRuntimeHLSLElementWiseBinaryOps.cpp
--- a/Engine/Plugins/Experimental/NNX/Source/NNXRuntimeRDG/Private/NNXRuntimeHLSLElementWiseBinaryOps.cpp
+++ b/Engine/Plugins/Experimental/NNX/Source/NNXRuntimeRDG/Private/NNXRuntimeHLSLElementWiseBinaryOps.cpp
@@ -6,6 +6,7 @@
 
 namespace NNX
 {
+DECLARE_GPU_STAT_NAMED(FMLHLSLOperatorElementWiseBinary, TEXT("FML.HLSL.Operator.ElementWise.Binary"));
 
 /**
  * Binary Element-wise ML operator implementation
@@ -32,7 +33,7 @@ private:
 
 public:
 
-	virtual bool Initialize(TArrayView<const FMLTensorDesc> InputTensors, TArrayView<const FMLTensorDesc> OutputTensors) override
+	virtual bool Initialize(TArrayView<const FMLTensorDesc> InputTensors, TArrayView<const FMLTensorDesc> OutputTensors, const FMLAttributeMap& Attributes) override
 	{
 		check(InputTensors.Num() == 2);
 		check(OutputTensors.Num() == 1);
@@ -70,10 +71,13 @@ public:
 		PermutationVector.Set<FMLElementWiseBinaryCS::FBinaryNumDimensions>(Output.Dimension);
 
 		TShaderMapRef<FMLElementWiseBinaryCS> ComputeShader(GetGlobalShaderMap(GMaxRHIFeatureLevel), PermutationVector);
+		
+		RDG_EVENT_SCOPE(GraphBuilder, "FML.HLSL.Operator.ElementWise.Binary");
+		RDG_GPU_STAT_SCOPE(GraphBuilder, FMLHLSLOperatorElementWiseBinary);
 
 		FComputeShaderUtils::AddPass(
 			GraphBuilder,
-			RDG_EVENT_NAME("FMLBinaryElementWiseOperatorHlsl_Dispatch"),
+			RDG_EVENT_NAME("FML.HLSL.Operator.ElementWise.Binary.Dispatch"),
 			ERDGPassFlags::Compute | ERDGPassFlags::NeverCull,
 			ComputeShader,
 			Params,
diff --git a/Engine/Plugins/Experimental/NNX/Source/NNXRuntimeRDG/Private/NNXRuntimeHLSLElementWiseUnaryOps.cpp b/Engine/Plugins/Experimental/NNX/Source/NNXRuntimeRDG/Private/NNXRuntimeHLSLElementWiseUnaryOps.cpp
--- a/Engine/Plugins/Experimental/NNX/Source/NNXRuntimeRDG/Private/NNXRuntimeHLSLElementWiseUnaryOps.cpp
+++ b/Engine/Plugins/Experimental/NNX/Source/NNXRuntimeRDG/Private/NNXRuntimeHLSLElementWiseUnaryOps.cpp
@@ -6,6 +6,7 @@
 
 namespace NNX
 {
+DECLARE_GPU_STAT_NAMED(FMLHLSLOperatorElementWiseUnary, TEXT("FML.HLSL.Operator.ElementWise.Unary"));
 
 /**
  * Unary Element-wise ML operator implementation
@@ -33,7 +34,7 @@ private:
 
 public:
 
-	virtual bool Initialize(TArrayView<const FMLTensorDesc> InputTensors, TArrayView<const FMLTensorDesc> OutputTensors) override
+	virtual bool Initialize(TArrayView<const FMLTensorDesc> InputTensors, TArrayView<const FMLTensorDesc> OutputTensors, const FMLAttributeMap& Attributes) override
 	{
 		check(InputTensors.Num() == 1);
 		check(OutputTensors.Num() == 1);
@@ -41,7 +42,9 @@ public:
 		Input = InputTensors[0];
 		Output = OutputTensors[0];
 
-		/*TODO attributes*/
+		Alpha = Attributes.GetOptionalFloat(TEXT("alpha"), Alpha);
+		Beta = Attributes.GetOptionalFloat(TEXT("beta"), Beta);
+		Gamma = Attributes.GetOptionalFloat(TEXT("gamma"), Gamma);
 
 		return true;
 	}
@@ -70,9 +73,12 @@ public:
 
 		TShaderMapRef<FMLElementWiseCS> ComputeShader(GetGlobalShaderMap(GMaxRHIFeatureLevel), PermutationVector);
 
+		RDG_EVENT_SCOPE(GraphBuilder, "FML.HLSL.Operator.ElementWise.Unary");
+		RDG_GPU_STAT_SCOPE(GraphBuilder, FMLHLSLOperatorElementWiseUnary);
+		
 		FComputeShaderUtils::AddPass(
 			GraphBuilder,
-			RDG_EVENT_NAME("FMLUnaryElementWiseOperatorHlsl_Dispatch"),
+			RDG_EVENT_NAME("FML.HLSL.Operator.ElementWise.Unary.Dispatch"),
 			ERDGPassFlags::Compute | ERDGPassFlags::NeverCull,
 			ComputeShader,
 			Params,
diff --git a/Engine/Plugins/Experimental/NNX/Source/NNXRuntimeRDG/Private/NNXRuntimeHLSLElementWiseVariadicOps.cpp b/Engine/Plugins/Experimental/NNX/Source/NNXRuntimeRDG/Private/NNXRuntimeHLSLElementWiseVariadicOps.cpp
--- a/Engine/Plugins/Experimental/NNX/Source/NNXRuntimeRDG/Private/NNXRuntimeHLSLElementWiseVariadicOps.cpp
+++ b/Engine/Plugins/Experimental/NNX/Source/NNXRuntimeRDG/Private/NNXRuntimeHLSLElementWiseVariadicOps.cpp
@@ -6,7 +6,7 @@
 
 namespace NNX
 {
-	DECLARE_GPU_STAT_NAMED(VariadicElementWiseOperatorHlsl, TEXT("NNX VariadicElementWiseOperatorHlsl"));
+	DECLARE_GPU_STAT_NAMED(FMLHLSLOperatorElementWiseVariadic, TEXT("FML.HLSL.Operator.ElementWise.Variadic"));
 
 	void AddOneVariadicOpPass(FRDGBuilder& GraphBuilder, 
 		TArrayView<const FMLTensorBinding> InputBindings,
@@ -71,7 +71,7 @@ namespace NNX
 		TShaderMapRef<FMLElementWiseVariadicCS> ComputeShader(GetGlobalShaderMap(GMaxRHIFeatureLevel), PermutationVector);
 		FComputeShaderUtils::AddPass(
 			GraphBuilder,
-			RDG_EVENT_NAME("FMLVariadicElementWiseOperatorHlsl_Dispatch"),
+			RDG_EVENT_NAME("FML.HLSL.Operator.ElementWise.Variadic.Dispatch"),
 			ERDGPassFlags::Compute | ERDGPassFlags::NeverCull,
 			ComputeShader,
 			Params,
@@ -102,7 +102,7 @@ private:
 
 public:
 
-	virtual bool Initialize(TArrayView<const FMLTensorDesc> InputTensors, TArrayView<const FMLTensorDesc> OutputTensors) override
+	virtual bool Initialize(TArrayView<const FMLTensorDesc> InputTensors, TArrayView<const FMLTensorDesc> OutputTensors, const FMLAttributeMap& Attributes) override
 	{
 		check(InputTensors.Num() > 0);
 		check(OutputTensors.Num() == 1);
@@ -118,8 +118,8 @@ public:
 		check(OutOutputBindings.Num() == 1);
 		check(InInputBindings.Num() == InputDescs.Num());
 
-		RDG_EVENT_SCOPE(GraphBuilder, "VariadicElementWiseOperatorHlsl");
-		RDG_GPU_STAT_SCOPE(GraphBuilder, VariadicElementWiseOperatorHlsl);
+		RDG_EVENT_SCOPE(GraphBuilder, "FML.HLSL.Operator.ElementWise.Variadic");
+		RDG_GPU_STAT_SCOPE(GraphBuilder, FMLHLSLOperatorElementWiseVariadic);
 
 		FMLTensorBinding PassInputBindings[FMLElementWiseVariadicCS::MAX_NUM_INPUT];
 		FMLTensorDesc PassInputDescs[FMLElementWiseVariadicCS::MAX_NUM_INPUT];
diff --git a/Engine/Plugins/Experimental/NNX/Source/NNXRuntimeRDG/Private/NNXRuntimeHLSLGemmOp.cpp b/Engine/Plugins/Experimental/NNX/Source/NNXRuntimeRDG/Private/NNXRuntimeHLSLGemmOp.cpp
--- a/Engine/Plugins/Experimental/NNX/Source/NNXRuntimeRDG/Private/NNXRuntimeHLSLGemmOp.cpp
+++ b/Engine/Plugins/Experimental/NNX/Source/NNXRuntimeRDG/Private/NNXRuntimeHLSLGemmOp.cpp
@@ -6,6 +6,7 @@
 
 namespace NNX
 {
+	DECLARE_GPU_STAT_NAMED(FMLHLSLOperatorGemm, TEXT("FML.HLSL.Operator.Gemm"));
 
 	/**
 	 * Gemm ML operator implementation
@@ -45,7 +46,7 @@ namespace NNX
 
 	public:
 
-		virtual bool Initialize(TArrayView<const FMLTensorDesc> InputTensors, TArrayView<const FMLTensorDesc> OutputTensors) override
+		virtual bool Initialize(TArrayView<const FMLTensorDesc> InputTensors, TArrayView<const FMLTensorDesc> OutputTensors, const FMLAttributeMap& Attributes) override
 		{
 			check(InputTensors.Num() >= 2 && InputTensors.Num() <= 3);
 			check(OutputTensors.Num() == 1);
@@ -93,9 +94,12 @@ namespace NNX
 
 			FIntVector ThreadGroupCount = FMLGemmCS::GetGroupCount(*Parameters, Algorithm, 0);
 
+			RDG_EVENT_SCOPE(GraphBuilder, "FML.HLSL.Operator.Gemm");
+			RDG_GPU_STAT_SCOPE(GraphBuilder, FMLHLSLOperatorGemm);
+
 			FComputeShaderUtils::AddPass(
 				GraphBuilder,
-				RDG_EVENT_NAME("FMLHlslGemmOperatorHlsl_Dispatch"),
+				RDG_EVENT_NAME("FML.HLSL.Operator.Gemm.Dispatch"),
 				ERDGPassFlags::Compute | ERDGPassFlags::NeverCull,
 				ComputeShader,
 				Parameters,
diff --git a/Engine/Plugins/Experimental/NNX/Source/NNXRuntimeRDG/Private/NNXRuntimeHLSLMatMulOp.cpp b/Engine/Plugins/Experimental/NNX/Source/NNXRuntimeRDG/Private/NNXRuntimeHLSLMatMulOp.cpp
--- a/Engine/Plugins/Experimental/NNX/Source/NNXRuntimeRDG/Private/NNXRuntimeHLSLMatMulOp.cpp
+++ b/Engine/Plugins/Experimental/NNX/Source/NNXRuntimeRDG/Private/NNXRuntimeHLSLMatMulOp.cpp
@@ -6,6 +6,7 @@
 
 namespace NNX
 {
+    DECLARE_GPU_STAT_NAMED(FMLHLSLOperatorMatMul, TEXT("FML.HLSL.Operator.MatMul"));
 
 	/**
 	 * Gemm ML operator implementation
@@ -31,7 +32,7 @@ namespace NNX
 
 	public:
 
-		virtual bool Initialize(TArrayView<const FMLTensorDesc> InputTensors, TArrayView<const FMLTensorDesc> OutputTensors) override
+		virtual bool Initialize(TArrayView<const FMLTensorDesc> InputTensors, TArrayView<const FMLTensorDesc> OutputTensors, const FMLAttributeMap& Attributes) override
 		{
 			check(InputTensors.Num() == 2);
 			check(OutputTensors.Num() == 1);
@@ -64,9 +65,12 @@ namespace NNX
 
 			FIntVector ThreadGroupCount = FMLGemmCS::GetGroupCount(*Parameters, Algorithm, NumStackDimensions);
 
+			RDG_EVENT_SCOPE(GraphBuilder, "FML.HLSL.Operator.MatMul");
+			RDG_GPU_STAT_SCOPE(GraphBuilder, FMLHLSLOperatorMatMul);
+
 			FComputeShaderUtils::AddPass(
 				GraphBuilder,
-				RDG_EVENT_NAME("FMLHlslMatMulOperatorHlsl_Dispatch"),
+				RDG_EVENT_NAME("FML.HLSL.Operator.MatMul.Dispatch"),
 				ERDGPassFlags::Compute | ERDGPassFlags::NeverCull,
 				ComputeShader,
 				Parameters,
diff --git a/Engine/Plugins/Experimental/NNX/Source/NNXRuntimeRDG/Private/NNXRuntimeHLSLOp.h b/Engine/Plugins/Experimental/NNX/Source/NNXRuntimeRDG/Private/NNXRuntimeHLSLOp.h
--- a/Engine/Plugins/Experimental/NNX/Source/NNXRuntimeRDG/Private/NNXRuntimeHLSLOp.h
+++ b/Engine/Plugins/Experimental/NNX/Source/NNXRuntimeRDG/Private/NNXRuntimeHLSLOp.h
@@ -13,7 +13,7 @@ namespace NNX
 	struct FMLOperatorHlsl : public FMLOperatorRDG
 	{
 		virtual ~FMLOperatorHlsl() = default;
-		virtual bool Initialize(TArrayView<const FMLTensorDesc> InputTensors, TArrayView<const FMLTensorDesc> OutputTensors) = 0;
+		virtual bool Initialize(TArrayView<const FMLTensorDesc> InputTensors, TArrayView<const FMLTensorDesc> OutputTensors, const FMLAttributeMap& Attributes) = 0;
 		virtual void Dispatch(FRDGBuilder& GraphBuilder, TArrayView<const FMLTensorBinding> InInputBindings, TArrayView<const FMLTensorBinding> OutOutputBindings) = 0;
 	};
 
diff --git a/Engine/Plugins/Experimental/NNX/Source/NNXUtils/Private/NNXModelBuilderNNX.cpp b/Engine/Plugins/Experimental/NNX/Source/NNXUtils/Private/NNXModelBuilderNNX.cpp
--- a/Engine/Plugins/Experimental/NNX/Source/NNXUtils/Private/NNXModelBuilderNNX.cpp
+++ b/Engine/Plugins/Experimental/NNX/Source/NNXUtils/Private/NNXModelBuilderNNX.cpp
@@ -170,7 +170,6 @@ public:
 
 		Format.Operators[OpIdx].InTensors.Emplace(TensorIdx);
 
-
 		return true;
 	}
 
@@ -187,6 +186,18 @@ public:
 		return true;
 	}
 
+	/** Add operator attribute */
+	virtual bool AddOperatorAttribute(HOperator Op, const FString& Name, const FMLAttributeValue& Value) override
+	{
+		int OpIdx = NNXOperatorCast(Op);
+
+		FMLFormatAttributeDesc& Attribute = Format.Operators[OpIdx].Attributes.Emplace_GetRef();
+		Attribute.Name = Name;
+		Attribute.Value = Value;
+		
+		return true;
+	}
+
 private:
 
 	int AddTensor(const FString& InName, TArrayView<const int32> InShape, EMLTensorDataType InDataType, EMLFormatTensorType InType)
diff --git a/Engine/Plugins/Experimental/NNX/Source/NNXUtils/Private/NNXModelBuilderONNX.cpp b/Engine/Plugins/Experimental/NNX/Source/NNXUtils/Private/NNXModelBuilderONNX.cpp
--- a/Engine/Plugins/Experimental/NNX/Source/NNXUtils/Private/NNXModelBuilderONNX.cpp
+++ b/Engine/Plugins/Experimental/NNX/Source/NNXUtils/Private/NNXModelBuilderONNX.cpp
@@ -199,6 +199,31 @@ public:
 		return true;
 	}
 
+	/** Add operator attribute */
+	virtual bool AddOperatorAttribute(HOperator Op, const FString& Name, const FMLAttributeValue& Value) override
+	{
+		auto NodeOp = OnnxOperatorCast(Op);
+		onnx::AttributeProto* Attribute = NodeOp->mutable_attribute()->Add();
+
+		Attribute->set_name(TCHAR_TO_ANSI(*Name));
+		if (Value.GetType() == EMLAttributeDataType::Float)
+		{
+			Attribute->set_type(onnx::AttributeProto::FLOAT);
+			Attribute->set_f(Value.AsFloat());
+		}
+		else if (Value.GetType() == EMLAttributeDataType::Int32)
+		{
+			Attribute->set_type(onnx::AttributeProto::INT);
+			Attribute->set_i(Value.AsInt32());
+		}
+		else
+		{
+			checkf(false, TEXT("not implemented"))
+		}
+		
+		return true;
+	}
+
 private:
 
 	bool SetValue(onnx::ValueInfoProto* Value, const FString& Name, EMLTensorDataType DataType, const TArrayView<const int32>& InShape)
@@ -249,7 +274,16 @@ private:
 //
 //
 //
-NNXUTILS_API bool CreateONNXModelForOperator(const FString& OperatorName, TArrayView<FMLTensorDesc> InInputTensors, TArrayView<FMLTensorDesc> InOutputTensors, TArray<uint8>& ModelData)
+NNXUTILS_API bool CreateONNXModelForOperator(const FString& OperatorName, TConstArrayView<FMLTensorDesc> InInputTensors, TConstArrayView<FMLTensorDesc> InOutputTensors, TArray<uint8>& ModelData)
+{
+	FMLAttributeMap EmptyAttributeMap;
+	return CreateONNXModelForOperator(OperatorName, InInputTensors, InOutputTensors, EmptyAttributeMap, ModelData);
+}
+
+//
+//
+//
+NNXUTILS_API bool CreateONNXModelForOperator(const FString& OperatorName, TConstArrayView<FMLTensorDesc> InInputTensors, TConstArrayView<FMLTensorDesc> InOutputTensors, const FMLAttributeMap& Attributes, TArray<uint8>& ModelData)
 {
 	TUniquePtr<IMLModelBuilder> Builder(CreateONNXModelBuilder());
 
@@ -289,6 +323,11 @@ NNXUTILS_API bool CreateONNXModelForOperator(const FString& OperatorName, TArray
 		Builder->AddOperatorOutput(Op, OutputTensors[Idx]);
 	}
 
+	for (int32 Idx = 0; Idx < Attributes.Num(); ++Idx)
+	{
+		Builder->AddOperatorAttribute(Op, Attributes.GetName(Idx), Attributes.GetAttributeValue(Idx));
+	}
+
 	Builder->End(ModelData);
 
 	return true;
diff --git a/Engine/Plugins/Experimental/NNX/Source/NNXUtils/Private/NNXModelOptimizer.cpp b/Engine/Plugins/Experimental/NNX/Source/NNXUtils/Private/NNXModelOptimizer.cpp
--- a/Engine/Plugins/Experimental/NNX/Source/NNXUtils/Private/NNXModelOptimizer.cpp
+++ b/Engine/Plugins/Experimental/NNX/Source/NNXUtils/Private/NNXModelOptimizer.cpp
@@ -344,6 +344,26 @@ public:
 
 			auto Op = Builder->AddOperator(NodeInfo.opName);
 
+			for (int InIdx = 0; InIdx < NodeInfo.attributeCount; ++InIdx)
+			{
+				Ort::GraphAttributeInfo AttrInfo = Graph->GetNodeAttribute(Node, InIdx);
+				Ort::GraphAttributeValue AttrValue = Graph->GetNodeAttributeValue(Node, InIdx);
+
+				if (AttrInfo.type == Ort::GraphAttributeType::kFloat)
+				{
+					Builder->AddOperatorAttribute(Op, AttrInfo.name, FMLAttributeValue(AttrValue.f));
+				}
+				else if (AttrInfo.type == Ort::GraphAttributeType::kInt)
+				{
+					Builder->AddOperatorAttribute(Op, AttrInfo.name, FMLAttributeValue((int)AttrValue.i));
+				}
+				else
+				{
+					//TODO: better error reporting add type (example: sparse tensor) and name of the actual node if any (not the op name but the node one)
+					UE_LOG(LogNNX, Warning, TEXT("Unsupported attribute type for attribute '%s' in node '%s'"), ANSI_TO_TCHAR(AttrInfo.name), ANSI_TO_TCHAR(NodeInfo.opName));
+				}
+			}
+
 			for (int InIdx = 0; InIdx < NodeInfo.inputCount; ++InIdx)
 			{
 				Ort::GraphTensorInfo		TensorInfo = Graph->GetNodeInput(Node, InIdx);
diff --git a/Engine/Plugins/Experimental/NNX/Source/NNXUtils/Public/NNXModelBuilder.h b/Engine/Plugins/Experimental/NNX/Source/NNXUtils/Public/NNXModelBuilder.h
--- a/Engine/Plugins/Experimental/NNX/Source/NNXUtils/Public/NNXModelBuilder.h
+++ b/Engine/Plugins/Experimental/NNX/Source/NNXUtils/Public/NNXModelBuilder.h
@@ -63,6 +63,9 @@ public:
 	/** Add operator input */
 	virtual bool AddOperatorInput(HOperator Op, HTensor Tensor) = 0;
 
+	/** Add operator attribute */
+	virtual bool AddOperatorAttribute(HOperator Op, const FString& Name, const FMLAttributeValue& Value) = 0;
+	
 	/** Add operator output */
 	virtual bool AddOperatorOutput(HOperator Op, HTensor Tensor) = 0;
 };
@@ -79,9 +82,10 @@ static constexpr int64 OnnxOpsetVersion = 15;
 NNXUTILS_API IMLModelBuilder* CreateONNXModelBuilder(int64 IrVersion = OnnxIrVersion, int64 OpsetVersion = OnnxOpsetVersion);
 
 /**
- * Utility function to create single layer NN for operator testing
+ * Utility functions to create single layer NN for operator testing with optional attributes
  */
-NNXUTILS_API bool CreateONNXModelForOperator(const FString& OperatorName, TArrayView<FMLTensorDesc> InputTensors, TArrayView<FMLTensorDesc> OutputTensors, TArray<uint8>& ModelData);
+NNXUTILS_API bool CreateONNXModelForOperator(const FString& OperatorName, TConstArrayView<const FMLTensorDesc> InInputTensors, TConstArrayView<const FMLTensorDesc> InOutputTensors, TArray<uint8>& ModelData);
+NNXUTILS_API bool CreateONNXModelForOperator(const FString& OperatorName, TConstArrayView<const FMLTensorDesc> InInputTensors, TConstArrayView<const FMLTensorDesc> InOutputTensors, const FMLAttributeMap& Attributes, TArray<uint8>& ModelData);
 
 /**
  * Create an instance of NNX model builder that creates NNX model/format in memory
