Fixed static analyzer warning in NNXGemmCS.cpp(62) : warning C6385: Reading invalid data from 'InputB.Sizes'...

diff --git a/Engine/Plugins/Experimental/NNX/Source/NNXHlslShaders/Private/NNXGemmCS.cpp b/Engine/Plugins/Experimental/NNX/Source/NNXHlslShaders/Private/NNXGemmCS.cpp
--- a/Engine/Plugins/Experimental/NNX/Source/NNXHlslShaders/Private/NNXGemmCS.cpp
+++ b/Engine/Plugins/Experimental/NNX/Source/NNXHlslShaders/Private/NNXGemmCS.cpp
@@ -58,8 +58,11 @@ namespace
 
 		for (int32 i = 0; i < NumStackDimensions; i++)
 		{
-			Result.StackShapeA[Result.StackShapeA.Num() - 1 - i] = i < NumStackDimensionsA ? InputA.Sizes[InputA.Dimension - 3 - i] : 1;
-			Result.StackShapeB[Result.StackShapeB.Num() - 1 - i] = i < NumStackDimensionsB ? InputB.Sizes[InputB.Dimension - 3 - i] : 1;
+			int32 IdxA = InputA.Dimension - 3 - i;
+			int32 IdxB = InputB.Dimension - 3 - i;
+
+			Result.StackShapeA[Result.StackShapeA.Num() - 1 - i] = IdxA >= 0 ? InputA.Sizes[IdxA] : 1;
+			Result.StackShapeB[Result.StackShapeB.Num() - 1 - i] = IdxB >= 0 ? InputB.Sizes[IdxB] : 1;
 		}
 
 		Result.StackStrideA.Init(1, NumStackDimensions);
