Move shader parameters from NNXShaderParameters.h to the HLSL shaders

diff --git a/Engine/Plugins/Experimental/NNX/Source/NNXHlslShaders/Private/NNXConvCS.cpp b/Engine/Plugins/Experimental/NNX/Source/NNXHlslShaders/Private/NNXConvCS.cpp
--- a/Engine/Plugins/Experimental/NNX/Source/NNXHlslShaders/Private/NNXConvCS.cpp
+++ b/Engine/Plugins/Experimental/NNX/Source/NNXHlslShaders/Private/NNXConvCS.cpp
@@ -104,7 +104,7 @@ TArray<int32> GetGridShape(TArray<int32> YShape, TArray<int32> GroupShape)
 void FMLConvCS::ModifyCompilationEnvironment(const FGlobalShaderPermutationParameters& InParameters, FShaderCompilerEnvironment& OutEnvironment)
 {
 	FGlobalShader::ModifyCompilationEnvironment(InParameters, OutEnvironment);
-	OutEnvironment.SetDefine(TEXT("MAX_NUM_DIMENSIONS"), NNXRT_CONV_MAX_NUM_DIMENSIONS);
+	OutEnvironment.SetDefine(TEXT("MAX_NUM_DIMENSIONS"), FConvConstants::MAX_NUM_DIMENSIONS);
 }
 
 TArray<int32> FMLConvCS::GetOutputShape(TArray<int32> XShape, TArray<int32> WShape, EConvAutoPad AutoPad, TArray<int32> Dilations, TArray<int32> Strides, TArray<int32> Pads)
diff --git a/Engine/Plugins/Experimental/NNX/Source/NNXHlslShaders/Private/NNXConvTransposeCS.cpp b/Engine/Plugins/Experimental/NNX/Source/NNXHlslShaders/Private/NNXConvTransposeCS.cpp
--- a/Engine/Plugins/Experimental/NNX/Source/NNXHlslShaders/Private/NNXConvTransposeCS.cpp
+++ b/Engine/Plugins/Experimental/NNX/Source/NNXHlslShaders/Private/NNXConvTransposeCS.cpp
@@ -5,7 +5,7 @@
 void FMLConvTransposeCS::ModifyCompilationEnvironment(const FGlobalShaderPermutationParameters& InParameters, FShaderCompilerEnvironment& OutEnvironment)
 {
 	FGlobalShader::ModifyCompilationEnvironment(InParameters, OutEnvironment);
-	OutEnvironment.SetDefine(TEXT("MAX_NUM_STACK_DIMENSIONS"), NNXRT_CONVTRANSPOSE_MAX_NUM_STACK_DIMENSIONS);
+	OutEnvironment.SetDefine(TEXT("MAX_NUM_STACK_DIMENSIONS"), FConvTransposeConstants::MAX_NUM_DIMENSIONS);
 }
 
 TArray<int32> FMLConvTransposeCS::GetOutputShape(TArray<int32> XShape, TArray<int32> WShape, EConvTransposeAutoPad AutoPad, TArray<int32> Dilations, TArray<int32> Strides, TArray<int32> Pads, TArray<int32> OutputPadding, int32 Group)
@@ -123,9 +123,9 @@ int32 FMLConvTransposeCS::GetNumReadsPerThread(EConvTransposeGroupSize GroupSize
 	}
 
 	int32 NumReads = FMath::DivideAndRoundUp(NumXBlockElements, NumThreadsPerGroup);
-	int32 NumReadsPow2 = FMath::Max(FMath::RoundToPositiveInfinity(FMath::Log2((float)NumReads)), NNXRT_CONVTRANSPOSE_MIN_NUM_READS_PER_THREAD_POW2);
+	int32 NumReadsPow2 = FMath::Max(FMath::RoundToPositiveInfinity(FMath::Log2((float)NumReads)), FConvTransposeConstants::MIN_NUM_READS_PER_THREAD_POW2);
 
-	if (NumReadsPow2 <= NNXRT_CONVTRANSPOSE_MAX_NUM_READS_PER_THREAD_POW2)
+	if (NumReadsPow2 <= FConvTransposeConstants::MAX_NUM_READS_PER_THREAD_POW2)
 	{
 		return NumReadsPow2;
 	}
diff --git a/Engine/Plugins/Experimental/NNX/Source/NNXHlslShaders/Private/NNXGemmCS.cpp b/Engine/Plugins/Experimental/NNX/Source/NNXHlslShaders/Private/NNXGemmCS.cpp
--- a/Engine/Plugins/Experimental/NNX/Source/NNXHlslShaders/Private/NNXGemmCS.cpp
+++ b/Engine/Plugins/Experimental/NNX/Source/NNXHlslShaders/Private/NNXGemmCS.cpp
@@ -29,7 +29,7 @@ namespace
 		int32 NumStackDimensionsB = FMath::Max((int32)InputB.Dimension - 2, 0);
 		int32 NumStackDimensions = FMath::Max(NumStackDimensionsA, NumStackDimensionsB);
 
-		check(NumStackDimensions <= NNXRT_GEMM_MAX_NUM_STACK_DIMENSIONS);
+		check(NumStackDimensions <= FGemmConstants::MAX_NUM_STACK_DIMENSIONS);
 
 		// Check matrix stack dimensions
 		if (NumStackDimensionsA > 0 && NumStackDimensionsB > 0)
@@ -84,7 +84,7 @@ namespace
 void FMLGemmCS::ModifyCompilationEnvironment(const FGlobalShaderPermutationParameters& InParameters, FShaderCompilerEnvironment& OutEnvironment)
 {
 	FGlobalShader::ModifyCompilationEnvironment(InParameters, OutEnvironment);
-	OutEnvironment.SetDefine(TEXT("MAX_NUM_STACK_DIMENSIONS"), NNXRT_GEMM_MAX_NUM_STACK_DIMENSIONS);
+	OutEnvironment.SetDefine(TEXT("MAX_NUM_STACK_DIMENSIONS"), FGemmConstants::MAX_NUM_STACK_DIMENSIONS);
 }
 
 void FMLGemmCS::FillInParameters(float Alpha, float Beta, int32 TransA, int32 TransB, const NNX::FMLTensorDesc &InputA, const NNX::FMLTensorDesc &InputB,
diff --git a/Engine/Plugins/Experimental/NNX/Source/NNXHlslShaders/Public/NNXConvCS.h b/Engine/Plugins/Experimental/NNX/Source/NNXHlslShaders/Public/NNXConvCS.h
--- a/Engine/Plugins/Experimental/NNX/Source/NNXHlslShaders/Public/NNXConvCS.h
+++ b/Engine/Plugins/Experimental/NNX/Source/NNXHlslShaders/Public/NNXConvCS.h
@@ -3,7 +3,8 @@
 #pragma once
 
 #include "GlobalShader.h"
-#include "NNXShaderParameters.h"
+#include "ShaderParameterUtils.h"
+#include "RenderGraphUtils.h"
 
 enum class EConvAlgorithm : uint8
 {
@@ -31,6 +32,7 @@ enum class EConvAutoPad : uint8
 class FConvConstants
 {
 public:
+	static const int32 MAX_NUM_DIMENSIONS{4};
 	static const int32 MIN_NUM_READS_PER_THREAD_POW2{1};
 	static const int32 MAX_NUM_READS_PER_THREAD_POW2{3};
 };
@@ -42,13 +44,35 @@ class NNXHLSLSHADERS_API FMLConvCS : public FGlobalShader
 
 	class FConvAlgorithm : SHADER_PERMUTATION_ENUM_CLASS("ALGORITHM", EConvAlgorithm);
 	class FConvGroupSize : SHADER_PERMUTATION_ENUM_CLASS("GROUP_SIZE", EConvGroupSize);
-	class FConvNumDimensions : SHADER_PERMUTATION_RANGE_INT("NUM_DIMENSIONS", 1, NNXRT_CONV_MAX_NUM_DIMENSIONS);
+	class FConvNumDimensions : SHADER_PERMUTATION_RANGE_INT("NUM_DIMENSIONS", 1, FConvConstants::MAX_NUM_DIMENSIONS);
 	class FConvNumReadsPerThread : SHADER_PERMUTATION_RANGE_INT("NUM_READS_PER_THREAD_POW2", FConvConstants::MIN_NUM_READS_PER_THREAD_POW2, FConvConstants::MAX_NUM_READS_PER_THREAD_POW2);
 	class FConvHasB : SHADER_PERMUTATION_BOOL("HAS_B");
 	using FPermutationDomain = TShaderPermutationDomain<FConvAlgorithm, FConvGroupSize, FConvNumDimensions, FConvNumReadsPerThread, FConvHasB>;
 
 public:
-	NNXRT_CONV_PARAMETER_STRUCT();
+	BEGIN_SHADER_PARAMETER_STRUCT(FParameters, )
+		SHADER_PARAMETER_RDG_BUFFER_SRV(Buffer<float>, X)
+		SHADER_PARAMETER_RDG_BUFFER_SRV(Buffer<float>, W)
+		SHADER_PARAMETER_RDG_BUFFER_UAV(RWBuffer<float>, Y)
+		SHADER_PARAMETER_RDG_BUFFER_SRV(Buffer<float>, B)
+		SHADER_PARAMETER_ARRAY(FIntVector4, Dilation_Stride_XBlockStartOffset_DilationXBlockStride, [FConvConstants::MAX_NUM_DIMENSIONS])
+		SHADER_PARAMETER_ARRAY(FIntVector4, GroupStride_GroupShape_GroupThreadStride_StrideXBlockStride, [FConvConstants::MAX_NUM_DIMENSIONS])
+		SHADER_PARAMETER_ARRAY(FIntVector4, YDimension_YMemoryStride_XDimension_XMemoryStride, [FConvConstants::MAX_NUM_DIMENSIONS])
+		SHADER_PARAMETER_ARRAY(FIntVector4, XBlockStartStride_XBlockStride_WDimension_WDimensionDilationXBlockStride, [FConvConstants::MAX_NUM_DIMENSIONS])
+		SHADER_PARAMETER_ARRAY(FVector4f, OneDiv_GroupStride_GroupThreadStride_XBlockStride, [FConvConstants::MAX_NUM_DIMENSIONS])
+		SHADER_PARAMETER(int32, NumWChannels)
+		SHADER_PARAMETER(int32, YBatchStride)
+		SHADER_PARAMETER(int32, YOutputKernelStride)
+		SHADER_PARAMETER(int32, XBatchStride)
+		SHADER_PARAMETER(int32, XChannelStride)
+		SHADER_PARAMETER(int32, XBlockSize)
+		SHADER_PARAMETER(int32, NumChannelBatches)
+		SHADER_PARAMETER(int32, NumChannelsPerBatch)
+		SHADER_PARAMETER(int32, WOutputKernelStride)
+		SHADER_PARAMETER(int32, WChannelBatchSize)
+		SHADER_PARAMETER(int32, WChannelSize)
+		SHADER_PARAMETER(float, GroupsDivM)
+	END_SHADER_PARAMETER_STRUCT()
 
 	static void ModifyCompilationEnvironment(const FGlobalShaderPermutationParameters& InParameters, FShaderCompilerEnvironment& OutEnvironment);
 
diff --git a/Engine/Plugins/Experimental/NNX/Source/NNXHlslShaders/Public/NNXConvTransposeCS.h b/Engine/Plugins/Experimental/NNX/Source/NNXHlslShaders/Public/NNXConvTransposeCS.h
--- a/Engine/Plugins/Experimental/NNX/Source/NNXHlslShaders/Public/NNXConvTransposeCS.h
+++ b/Engine/Plugins/Experimental/NNX/Source/NNXHlslShaders/Public/NNXConvTransposeCS.h
@@ -3,7 +3,8 @@
 #pragma once
 
 #include "GlobalShader.h"
-#include "NNXShaderParameters.h"
+#include "ShaderParameterUtils.h"
+#include "RenderGraphUtils.h"
 #include "NNXTypes.h"
 
 enum class EConvTransposeAlgorithm : uint8
@@ -29,6 +30,14 @@ enum class EConvTransposeAutoPad : uint8
 	MAX
 };
 
+class FConvTransposeConstants
+{
+public:
+    static const int32 MAX_NUM_DIMENSIONS{4};
+    static const int32 MIN_NUM_READS_PER_THREAD_POW2{1};
+    static const int32 MAX_NUM_READS_PER_THREAD_POW2{3};
+};
+
 class NNXHLSLSHADERS_API FMLConvTransposeCS : public FGlobalShader
 {
 	DECLARE_GLOBAL_SHADER(FMLConvTransposeCS);
@@ -36,13 +45,37 @@ class NNXHLSLSHADERS_API FMLConvTransposeCS : public FGlobalShader
 
 	class FConvTransposeAlgorithm : SHADER_PERMUTATION_ENUM_CLASS("ALGORITHM", EConvTransposeAlgorithm);
 	class FConvTransposeGroupSize : SHADER_PERMUTATION_ENUM_CLASS("GROUP_SIZE", EConvTransposeGroupSize);
-	class FConvTransposeNumStackDimensions : SHADER_PERMUTATION_RANGE_INT("NUM_STACK_DIMENSIONS", 1, NNXRT_CONVTRANSPOSE_MAX_NUM_STACK_DIMENSIONS);
-	class FConvTransposeNumReadsPerThread : SHADER_PERMUTATION_RANGE_INT("NUM_READS_PER_THREAD_POW2", NNXRT_CONVTRANSPOSE_MIN_NUM_READS_PER_THREAD_POW2, NNXRT_CONVTRANSPOSE_MAX_NUM_READS_PER_THREAD_POW2);
+	class FConvTransposeNumStackDimensions : SHADER_PERMUTATION_RANGE_INT("NUM_STACK_DIMENSIONS", 1, FConvTransposeConstants::MAX_NUM_DIMENSIONS);
+	class FConvTransposeNumReadsPerThread : SHADER_PERMUTATION_RANGE_INT("NUM_READS_PER_THREAD_POW2", FConvTransposeConstants::MIN_NUM_READS_PER_THREAD_POW2, FConvTransposeConstants::MAX_NUM_READS_PER_THREAD_POW2);
 	class FConvTransposeHasB : SHADER_PERMUTATION_BOOL("HAS_B");
 	using FPermutationDomain = TShaderPermutationDomain<FConvTransposeAlgorithm, FConvTransposeGroupSize, FConvTransposeNumStackDimensions, FConvTransposeNumReadsPerThread, FConvTransposeHasB>;
 
 public:
-	NNXRT_CONVTRANSPOSE_PARAMETER_STRUCT();
+	BEGIN_SHADER_PARAMETER_STRUCT(FParameters, )
+		SHADER_PARAMETER_RDG_BUFFER_SRV(Buffer<float>, X) \
+		SHADER_PARAMETER_RDG_BUFFER_SRV(Buffer<float>, W) \
+		SHADER_PARAMETER_RDG_BUFFER_UAV(RWBuffer<float>, Y) \
+		SHADER_PARAMETER_RDG_BUFFER_SRV(Buffer<float>, B) \
+		SHADER_PARAMETER_ARRAY(FIntVector4, Dilation_Stride_XBlockStartOffset_DilationXBlockStride, [FConvTransposeConstants::MAX_NUM_DIMENSIONS]) \
+		SHADER_PARAMETER_ARRAY(FIntVector4, GroupStride_GroupShape_GroupThreadStride_StrideXBlockStride, [FConvTransposeConstants::MAX_NUM_DIMENSIONS]) \
+		SHADER_PARAMETER_ARRAY(FIntVector4, YDimension_YMemoryStride_XDimension_XMemoryStride, [FConvTransposeConstants::MAX_NUM_DIMENSIONS]) \
+		SHADER_PARAMETER_ARRAY(FIntVector4, XBlockStartStride_XBlockStride_WDimension_WDimensionDilationXBlockStride, [FConvTransposeConstants::MAX_NUM_DIMENSIONS]) \
+		SHADER_PARAMETER_ARRAY(FVector4f, OneDiv_GroupStride_GroupThreadStride_OneDivStride, [FConvTransposeConstants::MAX_NUM_DIMENSIONS]) \
+		SHADER_PARAMETER(int32, NumWChannels) \
+		SHADER_PARAMETER(int32, NumOutChannelsDivGroup) \
+		SHADER_PARAMETER(int32, YBatchStride) \
+		SHADER_PARAMETER(int32, YOutputKernelStride) \
+		SHADER_PARAMETER(int32, XBatchStride) \
+		SHADER_PARAMETER(int32, XChannelStride) \
+		SHADER_PARAMETER(int32, XBlockSize) \
+		SHADER_PARAMETER(int32, NumChannelBatches) \
+		SHADER_PARAMETER(int32, NumChannelsPerBatch) \
+		SHADER_PARAMETER(int32, WOutputKernelStride) \
+		SHADER_PARAMETER(int32, WChannelBatchSize) \
+		SHADER_PARAMETER(int32, WChannelSize) \
+		SHADER_PARAMETER(float, GroupsDivM) \
+		SHADER_PARAMETER(float, OneDivGroup) \
+	END_SHADER_PARAMETER_STRUCT()
 
 	static void ModifyCompilationEnvironment(const FGlobalShaderPermutationParameters& InParameters, FShaderCompilerEnvironment& OutEnvironment);
 
diff --git a/Engine/Plugins/Experimental/NNX/Source/NNXHlslShaders/Public/NNXElementWiseBinaryCS.h b/Engine/Plugins/Experimental/NNX/Source/NNXHlslShaders/Public/NNXElementWiseBinaryCS.h
--- a/Engine/Plugins/Experimental/NNX/Source/NNXHlslShaders/Public/NNXElementWiseBinaryCS.h
+++ b/Engine/Plugins/Experimental/NNX/Source/NNXHlslShaders/Public/NNXElementWiseBinaryCS.h
@@ -3,10 +3,16 @@
 #pragma once
 
 #include "GlobalShader.h"
-#include "NNXShaderParameters.h"
-
+#include "ShaderParameterUtils.h"
+#include "RenderGraphUtils.h"
 #include "NNXOperator.h"
 
+class FElementWiseBinaryConstants
+{
+public:
+	static const int32 MAX_NUM_DIMENSIONS{8};
+};
+
 //
 //
 //
@@ -17,12 +23,19 @@ class NNXHLSLSHADERS_API FMLElementWiseBinaryCS : public FGlobalShader
 
 public:
 
-	NNXRT_ELEMENTWISEBINARY_PARAMETER_STRUCT()
+	BEGIN_SHADER_PARAMETER_STRUCT(FParameters, )
+		SHADER_PARAMETER_RDG_BUFFER_SRV(Buffer<float>, LHSInput) \
+		SHADER_PARAMETER_RDG_BUFFER_SRV(Buffer<float>, RHSInput) \
+		SHADER_PARAMETER_RDG_BUFFER_UAV(RWBuffer<float>, Output) \
+		SHADER_PARAMETER_ARRAY(FUintVector4, TensorInfo, [FElementWiseBinaryConstants::MAX_NUM_DIMENSIONS]) \
+		SHADER_PARAMETER(uint32, Num) \
+		SHADER_PARAMETER(uint32, ThreadCountX) \
+	END_SHADER_PARAMETER_STRUCT()
 	
 	static const uint32 THREADGROUP_SIZE_X;
 
 	class FOperatorType : SHADER_PERMUTATION_ENUM_CLASS("OP_TYPENAME", EMLElementWiseBinaryOperatorType);
-	class FBinaryNumDimensions : SHADER_PERMUTATION_RANGE_INT("NUM_DIMENSIONS", 1, NXRT_TENSORSTRIDEINFO_MAX_NUM_DIMENSIONS);
+	class FBinaryNumDimensions : SHADER_PERMUTATION_RANGE_INT("NUM_DIMENSIONS", 1, FElementWiseBinaryConstants::MAX_NUM_DIMENSIONS);
 	using FPermutationDomain = TShaderPermutationDomain<FOperatorType, FBinaryNumDimensions>;
 
 	static void ModifyCompilationEnvironment(const FGlobalShaderPermutationParameters& InParameters, FShaderCompilerEnvironment& OutEnvironment);
@@ -31,4 +44,3 @@ private:
 
 	static const FString GetOpFunc(EMLElementWiseBinaryOperatorType);
 };
-
diff --git a/Engine/Plugins/Experimental/NNX/Source/NNXHlslShaders/Public/NNXElementWiseCS.h b/Engine/Plugins/Experimental/NNX/Source/NNXHlslShaders/Public/NNXElementWiseCS.h
--- a/Engine/Plugins/Experimental/NNX/Source/NNXHlslShaders/Public/NNXElementWiseCS.h
+++ b/Engine/Plugins/Experimental/NNX/Source/NNXHlslShaders/Public/NNXElementWiseCS.h
@@ -3,8 +3,8 @@
 #pragma once
 
 #include "GlobalShader.h"
-#include "NNXShaderParameters.h"
-
+#include "ShaderParameterUtils.h"
+#include "RenderGraphUtils.h"
 #include "NNXOperator.h"
 
 //
@@ -17,7 +17,15 @@ class NNXHLSLSHADERS_API FMLElementWiseCS : public FGlobalShader
 
 public:
 
-	NNXRT_ELEMENTWISEUNARY_PARAMETER_STRUCT()
+	BEGIN_SHADER_PARAMETER_STRUCT(FParameters, )
+		SHADER_PARAMETER_RDG_BUFFER_SRV(Buffer<float>, Input)
+		SHADER_PARAMETER_RDG_BUFFER_UAV(RWBuffer<float>, Output)
+		SHADER_PARAMETER(uint32, Num)
+		SHADER_PARAMETER(uint32, ThreadCountX)
+		SHADER_PARAMETER(float, Alpha)
+		SHADER_PARAMETER(float, Beta)
+		SHADER_PARAMETER(float, Gamma)
+	END_SHADER_PARAMETER_STRUCT()
 
 	static const uint32 THREADGROUP_SIZE_X;
 
@@ -32,4 +40,3 @@ private:
 
 	static const FString GetOpFunc(EMLElementWiseUnaryOperatorType OpType);
 };
-
diff --git a/Engine/Plugins/Experimental/NNX/Source/NNXHlslShaders/Public/NNXElementWiseVariadicCS.h b/Engine/Plugins/Experimental/NNX/Source/NNXHlslShaders/Public/NNXElementWiseVariadicCS.h
--- a/Engine/Plugins/Experimental/NNX/Source/NNXHlslShaders/Public/NNXElementWiseVariadicCS.h
+++ b/Engine/Plugins/Experimental/NNX/Source/NNXHlslShaders/Public/NNXElementWiseVariadicCS.h
@@ -3,10 +3,16 @@
 #pragma once
 
 #include "GlobalShader.h"
-#include "NNXShaderParameters.h"
-
+#include "ShaderParameterUtils.h"
+#include "RenderGraphUtils.h"
 #include "NNXOperator.h"
 
+class FElementWiseVariadicConstants
+{
+public:
+	static const int32 MAX_NUM_DIMENSIONS{8};
+};
+
 class NNXHLSLSHADERS_API FMLElementWiseVariadicCS : public FGlobalShader
 {
 	DECLARE_GLOBAL_SHADER(FMLElementWiseVariadicCS);
@@ -21,15 +27,25 @@ public:
 	class FApplyScale : SHADER_PERMUTATION_BOOL("APPLYSCALE");
 	class FOutputAsInput : SHADER_PERMUTATION_BOOL("OUTPUTASINPUT");
 	class FNumInput : SHADER_PERMUTATION_RANGE_INT("NUMINPUT", 1, MAX_NUM_INPUT);
-	class FVariadicNumDimensions : SHADER_PERMUTATION_RANGE_INT("NUM_DIMENSIONS", 1, NXRT_TENSORSTRIDEINFO_MAX_NUM_DIMENSIONS);
+	class FVariadicNumDimensions : SHADER_PERMUTATION_RANGE_INT("NUM_DIMENSIONS", 1, FElementWiseVariadicConstants::MAX_NUM_DIMENSIONS);
 	using FPermutationDomain = TShaderPermutationDomain<FOperatorType, FApplyScale, FOutputAsInput, FNumInput, FVariadicNumDimensions>;
 
 	static void ModifyCompilationEnvironment(const FGlobalShaderPermutationParameters& InParameters, FShaderCompilerEnvironment& OutEnvironment);
 
-	NNXRT_ELEMENTWISEVARIADIC_PARAMETER_STRUCT()
+	BEGIN_SHADER_PARAMETER_STRUCT(FParameters, )
+		SHADER_PARAMETER_RDG_BUFFER_SRV(Buffer<float>, Input0)
+		SHADER_PARAMETER_RDG_BUFFER_SRV(Buffer<float>, Input1)
+		SHADER_PARAMETER_RDG_BUFFER_SRV(Buffer<float>, Input2)
+		SHADER_PARAMETER_RDG_BUFFER_SRV(Buffer<float>, Input3)
+		SHADER_PARAMETER_RDG_BUFFER_UAV(RWBuffer<float>, Output)
+		SHADER_PARAMETER_ARRAY(FUintVector4, InputTensorInfo, [FElementWiseVariadicConstants::MAX_NUM_DIMENSIONS])
+		SHADER_PARAMETER_ARRAY(FUintVector4, OutputTensorInfo, [FElementWiseVariadicConstants::MAX_NUM_DIMENSIONS])
+		SHADER_PARAMETER(uint32, Num)
+		SHADER_PARAMETER(uint32, ThreadCountX)
+		SHADER_PARAMETER(float, Scale)
+	END_SHADER_PARAMETER_STRUCT()
 
 private:
 
 	static const FString GetOpFunc(EMLElementWiseVariadicOperatorType);
 };
-
diff --git a/Engine/Plugins/Experimental/NNX/Source/NNXHlslShaders/Public/NNXGemmCS.h b/Engine/Plugins/Experimental/NNX/Source/NNXHlslShaders/Public/NNXGemmCS.h
--- a/Engine/Plugins/Experimental/NNX/Source/NNXHlslShaders/Public/NNXGemmCS.h
+++ b/Engine/Plugins/Experimental/NNX/Source/NNXHlslShaders/Public/NNXGemmCS.h
@@ -3,7 +3,8 @@
 #pragma once
 
 #include "GlobalShader.h"
-#include "NNXShaderParameters.h"
+#include "ShaderParameterUtils.h"
+#include "RenderGraphUtils.h"
 #include "NNXTypes.h"
 
 enum class EGemmCScalar : uint8
@@ -34,6 +35,12 @@ enum class EGemmAlgorithm : uint8
 	MAX
 };
 
+class FGemmConstants
+{
+public:
+	static const int32 MAX_NUM_STACK_DIMENSIONS{8};
+};
+
 class NNXHLSLSHADERS_API FMLGemmCS : public FGlobalShader
 {
 	DECLARE_GLOBAL_SHADER(FMLGemmCS);
@@ -41,19 +48,38 @@ class NNXHLSLSHADERS_API FMLGemmCS : public FGlobalShader
 
 	class FGemmCScalar : SHADER_PERMUTATION_ENUM_CLASS("C_SCALAR", EGemmCScalar);
 	class FGemmAlgorithm : SHADER_PERMUTATION_ENUM_CLASS("ALGORITHM", EGemmAlgorithm);
-	class FGemmNumStackDimensions : SHADER_PERMUTATION_RANGE_INT("NUM_STACK_DIMENSIONS", 0, NNXRT_GEMM_MAX_NUM_STACK_DIMENSIONS);
+	class FGemmNumStackDimensions : SHADER_PERMUTATION_RANGE_INT("NUM_STACK_DIMENSIONS", 0, FGemmConstants::MAX_NUM_STACK_DIMENSIONS);
 	using FPermutationDomain = TShaderPermutationDomain<FGemmCScalar, FGemmAlgorithm, FGemmNumStackDimensions>;
 
 public:
-	NNXRT_GEMM_PARAMETER_STRUCT()
+	BEGIN_SHADER_PARAMETER_STRUCT(FParameters, )
+		SHADER_PARAMETER(float, Alpha)
+		SHADER_PARAMETER(float, Beta)
+		SHADER_PARAMETER(int32, TransA)
+		SHADER_PARAMETER(int32, TransB)
+		SHADER_PARAMETER(uint32, M)
+		SHADER_PARAMETER(uint32, N)
+		SHADER_PARAMETER(uint32, K)
+		SHADER_PARAMETER(uint32, MxK)
+		SHADER_PARAMETER(uint32, KxN)
+		SHADER_PARAMETER(uint32, MxN)
+		SHADER_PARAMETER(uint32, CWidth)
+		SHADER_PARAMETER(uint32, CHeight)
+		SHADER_PARAMETER(float, CScalar)
+		SHADER_PARAMETER_ARRAY(FUint32Vector4, StackShapeA_StackShapeB_StackStrideA_StackStrideB, [FGemmConstants::MAX_NUM_STACK_DIMENSIONS])
+		SHADER_PARAMETER_RDG_BUFFER_SRV(Buffer<float>, A)
+		SHADER_PARAMETER_RDG_BUFFER_SRV(Buffer<float>, B)
+		SHADER_PARAMETER_RDG_BUFFER_SRV(Buffer<float>, C)
+		SHADER_PARAMETER_RDG_BUFFER_UAV(RWBuffer<float>, Y)
+	END_SHADER_PARAMETER_STRUCT()
 
 	static void ModifyCompilationEnvironment(const FGlobalShaderPermutationParameters& InParameters, FShaderCompilerEnvironment& OutEnvironment);
 
 	static void FillInParameters(float Alpha, float Beta, int32 TransA, int32 TransB, const NNX::FMLTensorDesc &InputA, const NNX::FMLTensorDesc &InputB,
-		const NNX::FMLTensorDesc &InputC, float CScalar, FMLGemmCS::FParameters& Parameters);
+		const NNX::FMLTensorDesc &InputC, float CScalar, FParameters& Parameters);
 
-	static void FillInParametersMatMul(const NNX::FMLTensorDesc &InputA, const NNX::FMLTensorDesc &InputB, FMLGemmCS::FParameters& Parameters);
+	static void FillInParametersMatMul(const NNX::FMLTensorDesc &InputA, const NNX::FMLTensorDesc &InputB, FParameters& Parameters);
 
-	static FIntVector GetGroupCount(const FMLGemmCS::FParameters& Parameters, EGemmAlgorithm Algorithm, int32 NumStackDimensions);
-	static EGemmAlgorithm GetAlgorithm(const FMLGemmCS::FParameters& Parameters);
+	static FIntVector GetGroupCount(const FParameters& Parameters, EGemmAlgorithm Algorithm, int32 NumStackDimensions);
+	static EGemmAlgorithm GetAlgorithm(const FParameters& Parameters);
 };
diff --git a/Engine/Plugins/Experimental/NNX/Source/NNXHlslShaders/Public/NNXShaderParameters.h b/Engine/Plugins/Experimental/NNX/Source/NNXHlslShaders/Public/NNXShaderParameters.h
--- a/Engine/Plugins/Experimental/NNX/Source/NNXHlslShaders/Public/NNXShaderParameters.h
+++ b/Engine/Plugins/Experimental/NNX/Source/NNXHlslShaders/Public/NNXShaderParameters.h
@@ -1,11 +1,13 @@
 // Copyright Epic Games, Inc. All Rights Reserved.
 
-//#pragma once
+#pragma once
 
 #include "ShaderParameterUtils.h"
 #include "RenderGraphUtils.h"
 
 // HLSL runtime and DirectML runtime use different shader parameter access patterns
+#define NNXRT_RDG_DML
+
 #ifdef NNXRT_RDG_DML
 #define NNXRT_RDG_BUFFER_SRV(View, BufferName) RDG_BUFFER_ACCESS(BufferName, ERHIAccess::UAVCompute)
 #define NNXRT_RDG_BUFFER_UAV(View, BufferName) RDG_BUFFER_ACCESS(BufferName, ERHIAccess::UAVCompute)
diff --git a/Engine/Plugins/Experimental/NNX/Source/NNXRuntimeRDG/Private/NNXRuntimeDML.cpp b/Engine/Plugins/Experimental/NNX/Source/NNXRuntimeRDG/Private/NNXRuntimeDML.cpp
--- a/Engine/Plugins/Experimental/NNX/Source/NNXRuntimeRDG/Private/NNXRuntimeDML.cpp
+++ b/Engine/Plugins/Experimental/NNX/Source/NNXRuntimeRDG/Private/NNXRuntimeDML.cpp
@@ -11,14 +11,7 @@
 #include "RenderGraphUtils.h"
 
 // TODO: Use NNXOperator.h instead of HlslIdentityOp
-#define NNXRT_RDG_DML
 #include "NNXShaderParameters.h"
-#undef NNXRT_RDG_DML
-// This is to fix compilation warnings on unity builds
-#undef NNXRT_BEGIN_SHADER_PARAMETER_STRUCT
-#undef NNXRT_END_SHADER_PARAMETER_STRUCT
-#undef NNXRT_RDG_BUFFER_SRV
-#undef NNXRT_RDG_BUFFER_UAV
 
 #include "HAL/FileManager.h"
 
