Tensor desc array rework

diff --git a/Engine/Plugins/Experimental/NNX/Source/NNXCore/Public/NNXRuntimeFormat.h b/Engine/Plugins/Experimental/NNX/Source/NNXCore/Public/NNXRuntimeFormat.h
--- a/Engine/Plugins/Experimental/NNX/Source/NNXCore/Public/NNXRuntimeFormat.h
+++ b/Engine/Plugins/Experimental/NNX/Source/NNXCore/Public/NNXRuntimeFormat.h
@@ -64,47 +64,6 @@ struct FMLFormatOperatorDesc
 	*/
 };
 
-USTRUCT()
-struct FMLFormatTensorShapeDesc
-{
-	static constexpr int32 MaxDimension = NNX::FMLTensorDesc::MaxTensorDimension;
-
-	GENERATED_USTRUCT_BODY()
-
-	UPROPERTY(VisibleAnywhere, Category = "Neural Network Inference")
-	uint32	Sizes[MaxDimension];
-
-	UPROPERTY(VisibleAnywhere, Category = "Neural Network Inference")
-	uint32	Dimension;
-
-	template<typename T>
-	void operator=(TArrayView<const T> OtherShape)
-	{
-		Dimension = OtherShape.Num();
-		if (Dimension > MaxDimension)
-		{
-			Dimension = MaxDimension;
-		}
-
-		for (uint32 Idx = 0; Idx < Dimension; ++Idx)
-		{
-			Sizes[Idx] = static_cast<uint32>(OtherShape[Idx]);
-		}
-	}
-
-	int32 Volume() const
-	{
-		int32 Result = 1;
-
-		for (uint32 Idx = 0; Idx < Dimension; ++Idx)
-		{
-			Result *= Sizes[Idx];
-		}
-
-		return Result;
-	}
-};
-
 USTRUCT()
 struct FMLFormatTensorDesc
 {
@@ -114,7 +73,7 @@ struct FMLFormatTensorDesc
 	FString Name;
 
 	UPROPERTY(VisibleAnywhere, Category = "Neural Network Inference")
-	FMLFormatTensorShapeDesc	Shape;
+	TArray<uint32> Shape;
 
 	UPROPERTY(VisibleAnywhere, Category = "Neural Network Inference")
 	EMLFormatTensorType	Type;
diff --git a/Engine/Plugins/Experimental/NNX/Source/NNXCore/Public/NNXTypes.h b/Engine/Plugins/Experimental/NNX/Source/NNXCore/Public/NNXTypes.h
--- a/Engine/Plugins/Experimental/NNX/Source/NNXCore/Public/NNXTypes.h
+++ b/Engine/Plugins/Experimental/NNX/Source/NNXCore/Public/NNXTypes.h
@@ -179,40 +179,24 @@ namespace NNX
  */
 struct FMLTensorDesc
 {
-	constexpr static uint32	MaxTensorDimension = 5;
+	constexpr static int32	MaxTensorDimension = 5;
 
 	FString					Name;
-	uint32					Dimension;							//!< Number of dimensions in the Sizes
-	uint32					Sizes[MaxTensorDimension];			//!< Sizes
+	TArray<uint32, TInlineAllocator<MaxTensorDimension>> Shape;	//!< Shape
 	uint64					DataSize;							//!< Size of data in bytes
 	EMLTensorDataType		DataType;
 
 	/** Make tensor descriptor */
 	static FMLTensorDesc Make(const FString& Name, TArrayView<const uint32> Shape, EMLTensorDataType DataType, uint64 DataSize = 0)
 	{
-		FMLTensorDesc Desc{};
-
-		uint32 Dimension = Shape.Num();
-
-		check(Dimension <= FMLTensorDesc::MaxTensorDimension);
-		if (Dimension > FMLTensorDesc::MaxTensorDimension)
+		check(Shape.Num() <= FMLTensorDesc::MaxTensorDimension);
+		if (Shape.Num() > FMLTensorDesc::MaxTensorDimension)
 		{
-			//UE_LOG(LogNNX, Warning, TEXT("Unsupported tensor dimension:%d"), Dimension);
-			return Desc;
+			//UE_LOG(LogNNX, Warning, TEXT("Unsupported tensor dimension:%d"), Shape.Num());
+			return {};
 		}
 
-		Desc.Name = Name;
-		Desc.Dimension = Dimension;
-		
-		for (uint32 Idx = 0; Idx < Desc.Dimension; ++Idx)
-		{
-			Desc.Sizes[Idx] = Shape[Idx];
-		}
-
-		Desc.DataType = DataType;
-		Desc.DataSize = DataSize;
-
-		return Desc;
+		return {Name, TArray<uint32, TInlineAllocator<MaxTensorDimension>>{Shape}, DataSize, DataType};
 	}
 
 	/** Check if descriptor is valid */
@@ -268,7 +252,7 @@ inline size_t GetTensorDataTypeSizeInBytes(EMLTensorDataType InType)
 /** Check if descriptor is valid */
 inline bool FMLTensorDesc::IsValid() const
 {
-	return DataType != EMLTensorDataType::None && Dimension > 0 && Dimension <= MaxTensorDimension; //&& Volume() >= 1;
+	return DataType != EMLTensorDataType::None && Shape.Num() > 0 && Shape.Num() <= MaxTensorDimension; //&& Volume() >= 1;
 }
 
 /** Return size of one element in bytes */
@@ -282,9 +266,9 @@ inline int32 FMLTensorDesc::Volume() const
 {
 	int32 Result = 1;
 
-	for (uint32 Idx = 0; Idx < Dimension; ++Idx)
+	for (int32 Idx = 0; Idx < Shape.Num(); ++Idx)
 	{
-		Result *= Sizes[Idx];
+		Result *= Shape[Idx];
 	}
 
 	return Result;
diff --git a/Engine/Plugins/Experimental/NNX/Source/NNXHlslShaders/Private/NNXGemmCS.cpp b/Engine/Plugins/Experimental/NNX/Source/NNXHlslShaders/Private/NNXGemmCS.cpp
--- a/Engine/Plugins/Experimental/NNX/Source/NNXHlslShaders/Private/NNXGemmCS.cpp
+++ b/Engine/Plugins/Experimental/NNX/Source/NNXHlslShaders/Private/NNXGemmCS.cpp
@@ -21,12 +21,12 @@ namespace
 
 	static GemmMatrixParameters GetMatrixParameters(const NNX::FMLTensorDesc &InputA, const NNX::FMLTensorDesc &InputB, const NNX::FMLTensorDesc &InputC)
 	{
-		check(InputA.Dimension != 0);
-		check(InputB.Dimension != 0);
+		check(InputA.Shape.Num() != 0);
+		check(InputB.Shape.Num() != 0);
 
 		GemmMatrixParameters Result;
-		int32 NumStackDimensionsA = FMath::Max((int32)InputA.Dimension - 2, 0);
-		int32 NumStackDimensionsB = FMath::Max((int32)InputB.Dimension - 2, 0);
+		int32 NumStackDimensionsA = FMath::Max((int32)InputA.Shape.Num() - 2, 0);
+		int32 NumStackDimensionsB = FMath::Max((int32)InputB.Shape.Num() - 2, 0);
 		int32 NumStackDimensions = FMath::Max(NumStackDimensionsA, NumStackDimensionsB);
 
 		check(NumStackDimensions <= FGemmConstants::MAX_NUM_STACK_DIMENSIONS);
@@ -37,32 +37,32 @@ namespace
 			const int32 NumDimToCheck = FMath::Min(NumStackDimensionsA, NumStackDimensionsB);
 			for (int32 i = 0; i < NumDimToCheck; i++)
 			{
-				const uint32 VolumeA = InputA.Sizes[NumStackDimensionsA - 1 - i];
-				const uint32 VolumeB = InputB.Sizes[NumStackDimensionsB - 1 - i];
+				const uint32 VolumeA = InputA.Shape[NumStackDimensionsA - 1 - i];
+				const uint32 VolumeB = InputB.Shape[NumStackDimensionsB - 1 - i];
 
 				check(VolumeA == 1 || VolumeB == 1 || VolumeA == VolumeB);
 			}
 		}
 
-		const bool IsVectorA = InputA.Dimension == 1;
-		const bool IsVectorB = InputB.Dimension == 1;
-		const bool IsVectorC = InputC.Dimension == 1;
+		const bool IsVectorA = InputA.Shape.Num() == 1;
+		const bool IsVectorB = InputB.Shape.Num() == 1;
+		const bool IsVectorC = InputC.Shape.Num() == 1;
 
-		Result.M = IsVectorA ? 1 : InputA.Sizes[InputA.Dimension - 2];
-		Result.N = IsVectorB ? 1 : InputB.Sizes[InputB.Dimension - 1];
-		Result.K = InputA.Sizes[IsVectorA ? 0 : InputA.Dimension - 1];
-		check(InputB.Sizes[IsVectorB ? 0 : InputB.Dimension - 2] == Result.K);
+		Result.M = IsVectorA ? 1 : InputA.Shape[InputA.Shape.Num() - 2];
+		Result.N = IsVectorB ? 1 : InputB.Shape[InputB.Shape.Num() - 1];
+		Result.K = InputA.Shape[IsVectorA ? 0 : InputA.Shape.Num() - 1];
+		check(InputB.Shape[IsVectorB ? 0 : InputB.Shape.Num() - 2] == Result.K);
 
 		Result.StackShapeA.Init(1, NumStackDimensions);
 		Result.StackShapeB.Init(1, NumStackDimensions);
 
 		for (int32 i = 0; i < NumStackDimensions; i++)
 		{
-			int32 IdxA = InputA.Dimension - 3 - i;
-			int32 IdxB = InputB.Dimension - 3 - i;
+			int32 IdxA = InputA.Shape.Num() - 3 - i;
+			int32 IdxB = InputB.Shape.Num() - 3 - i;
 
-			Result.StackShapeA[Result.StackShapeA.Num() - 1 - i] = IdxA >= 0 ? InputA.Sizes[IdxA] : 1;
-			Result.StackShapeB[Result.StackShapeB.Num() - 1 - i] = IdxB >= 0 ? InputB.Sizes[IdxB] : 1;
+			Result.StackShapeA[Result.StackShapeA.Num() - 1 - i] = IdxA >= 0 ? InputA.Shape[IdxA] : 1;
+			Result.StackShapeB[Result.StackShapeB.Num() - 1 - i] = IdxB >= 0 ? InputB.Shape[IdxB] : 1;
 		}
 
 		Result.StackStrideA.Init(1, NumStackDimensions);
@@ -74,8 +74,8 @@ namespace
 			Result.StackStrideB[i] = Result.StackStrideB[i + 1] * Result.StackShapeB[i + 1];
 		}
 
-		Result.CWidth = InputC.Dimension == 0 ? 0 : InputC.Sizes[InputC.Dimension - 1];
-		Result.CHeight = InputC.Dimension == 0 ? 0 : IsVectorC ? 1 : InputC.Sizes[InputC.Dimension - 2];
+		Result.CWidth = InputC.Shape.Num() == 0 ? 0 : InputC.Shape[InputC.Shape.Num() - 1];
+		Result.CHeight = InputC.Shape.Num() == 0 ? 0 : IsVectorC ? 1 : InputC.Shape[InputC.Shape.Num() - 2];
 
 		return Result;
 	}
diff --git a/Engine/Plugins/Experimental/NNX/Source/NNXQA/Private/NNXQAUtils.cpp b/Engine/Plugins/Experimental/NNX/Source/NNXQA/Private/NNXQAUtils.cpp
--- a/Engine/Plugins/Experimental/NNX/Source/NNXQA/Private/NNXQAUtils.cpp
+++ b/Engine/Plugins/Experimental/NNX/Source/NNXQA/Private/NNXQAUtils.cpp
@@ -86,7 +86,7 @@ namespace Test
 
 	FString FMLTensorDescToString(const FMLTensorDesc& desc)
 	{
-		TArrayView<const uint32> Shape(desc.Sizes, desc.Dimension);
+		TArrayView<const uint32> Shape(desc.Shape);
 
 		FString TensorDesc;
 		TensorDesc.Reserve(50);
@@ -247,13 +247,13 @@ namespace Test
 	{
 		bool bTensorDescMatch = true;
 		bTensorDescMatch &= (RefTensorDesc.Name == OtherTensorDesc.Name);
-		bTensorDescMatch &= (RefTensorDesc.Dimension == OtherTensorDesc.Dimension);
+		bTensorDescMatch &= (RefTensorDesc.Shape.Num() == OtherTensorDesc.Shape.Num());
 		bTensorDescMatch &= (RefTensorDesc.DataSize == OtherTensorDesc.DataSize);
 		bTensorDescMatch &= (RefTensorDesc.DataType == OtherTensorDesc.DataType);
-		check(RefTensorDesc.Dimension <= RefTensorDesc.MaxTensorDimension);
-		for (uint32 i = 0; i < RefTensorDesc.Dimension; ++i)
+		check(RefTensorDesc.Shape.Num() <= RefTensorDesc.MaxTensorDimension);
+		for (int32 i = 0; i < RefTensorDesc.Shape.Num(); ++i)
 		{
-			bTensorDescMatch &= (RefTensorDesc.Sizes[i] == OtherTensorDesc.Sizes[i]);
+			bTensorDescMatch &= (RefTensorDesc.Shape[i] == OtherTensorDesc.Shape[i]);
 		}
 
 		if (!bTensorDescMatch)
diff --git a/Engine/Plugins/Experimental/NNX/Source/NNXRuntimeCPU/Private/NNXRuntimeCPU.cpp b/Engine/Plugins/Experimental/NNX/Source/NNXRuntimeCPU/Private/NNXRuntimeCPU.cpp
--- a/Engine/Plugins/Experimental/NNX/Source/NNXRuntimeCPU/Private/NNXRuntimeCPU.cpp
+++ b/Engine/Plugins/Experimental/NNX/Source/NNXRuntimeCPU/Private/NNXRuntimeCPU.cpp
@@ -181,7 +181,7 @@ bool FMLInferenceModelCPU::ConfigureTensors(const bool InIsInput)
 		std::pair<EMLTensorDataType, uint64> TypeAndSize = TranslateTensorTypeORTToNNI(ONNXTensorElementDataTypeEnum);
 		CurrentTensorDescriptor.DataType = TypeAndSize.first;
 
-		CurrentTensorDescriptor.Dimension = CurrentTensorInfo.GetShape().size();
+		CurrentTensorDescriptor.Shape.SetNumUninitialized(CurrentTensorInfo.GetShape().size());
 		uint8 Index = 0;
 		uint64 TensorNumberElements = 1;
 		for (const int64_t CurrentTensorSize : CurrentTensorInfo.GetShape())
@@ -189,14 +189,14 @@ bool FMLInferenceModelCPU::ConfigureTensors(const bool InIsInput)
 			// Negative (variable) dimensions not implemented yet
 			if (CurrentTensorSize < 0)
 			{
-				CurrentTensorDescriptor.Sizes[Index] = 1;
+				CurrentTensorDescriptor.Shape[Index] = 1;
 				UE_LOG(LogNNX, Display,
 					TEXT("Negative (i.e., variable) dimensions not allowed yet, hard-coded to 1. Let us know if you really need variable dimensions."
 						" Keep in mind that fixed sizes might allow additional optimizations and speedup of the network during Run()."));
 			}
 			else
 			{
-				CurrentTensorDescriptor.Sizes[Index] = CurrentTensorSize;
+				CurrentTensorDescriptor.Shape[Index] = CurrentTensorSize;
 				TensorNumberElements *= CurrentTensorSize;
 			}
 			++Index;
diff --git a/Engine/Plugins/Experimental/NNX/Source/NNXRuntimeCPU/Public/NNXRuntimeCPUUtils.h b/Engine/Plugins/Experimental/NNX/Source/NNXRuntimeCPU/Public/NNXRuntimeCPUUtils.h
--- a/Engine/Plugins/Experimental/NNX/Source/NNXRuntimeCPU/Public/NNXRuntimeCPUUtils.h
+++ b/Engine/Plugins/Experimental/NNX/Source/NNXRuntimeCPU/Public/NNXRuntimeCPUUtils.h
@@ -135,14 +135,14 @@ inline void BindTensorsToORT(
 		ONNXTensorElementDataType CurrentORTType = InTensorsORTType[Index];
 
 		TUniquePtr<int64_t[]> SizesInt64t;
-		SizesInt64t = MakeUnique<int64_t[]>(CurrentDescriptor.Dimension);
-		for (uint32 DimIndex = 0; DimIndex < CurrentDescriptor.Dimension; ++DimIndex)
+		SizesInt64t = MakeUnique<int64_t[]>(CurrentDescriptor.Shape.Num());
+		for (int32 DimIndex = 0; DimIndex < CurrentDescriptor.Shape.Num(); ++DimIndex)
 		{
-			SizesInt64t.Get()[DimIndex] = CurrentDescriptor.Sizes[DimIndex];
+			SizesInt64t.Get()[DimIndex] = CurrentDescriptor.Shape[DimIndex];
 		}
 
 		const uint64 ByteCount { InTensorsDescriptors[Index].DataSize };
-		const uint32 ArrayDimensions { CurrentDescriptor.Dimension };
+		const uint32 ArrayDimensions { (uint32)CurrentDescriptor.Shape.Num() };
 		OutOrtTensors.Emplace(
 			Ort::Value::CreateTensor(
 				*InAllocatorInfo,
diff --git a/Engine/Plugins/Experimental/NNX/Source/NNXRuntimeORT/Private/NNXRuntimeORT.cpp b/Engine/Plugins/Experimental/NNX/Source/NNXRuntimeORT/Private/NNXRuntimeORT.cpp
--- a/Engine/Plugins/Experimental/NNX/Source/NNXRuntimeORT/Private/NNXRuntimeORT.cpp
+++ b/Engine/Plugins/Experimental/NNX/Source/NNXRuntimeORT/Private/NNXRuntimeORT.cpp
@@ -244,7 +244,7 @@ bool FMLInferenceModelORT::ConfigureTensors(const bool InIsInput)
 		std::pair<EMLTensorDataType, uint64> TypeAndSize = TranslateTensorTypeORTToNNI(ONNXTensorElementDataTypeEnum);
 		CurrentTensorDescriptor.DataType = TypeAndSize.first;
 		
-		CurrentTensorDescriptor.Dimension = CurrentTensorInfo.GetShape().size();
+		CurrentTensorDescriptor.Shape.SetNumUninitialized(CurrentTensorInfo.GetShape().size());
 		uint8 Index = 0;
 		uint64 TensorNumberElements = 1;
 		for (const int64_t CurrentTensorSize : CurrentTensorInfo.GetShape())
@@ -252,14 +252,14 @@ bool FMLInferenceModelORT::ConfigureTensors(const bool InIsInput)
 			// Negative (variable) dimensions not implemented yet
 			if (CurrentTensorSize < 0)
 			{
-				CurrentTensorDescriptor.Sizes[Index] = 1;
+				CurrentTensorDescriptor.Shape[Index] = 1;
 				UE_LOG(LogNNX, Display,
 					TEXT("Negative (i.e., variable) dimensions not allowed yet, hard-coded to 1. Let us know if you really need variable dimensions."
 						" Keep in mind that fixed sizes might allow additional optimizations and speedup of the network during Run()."));
 			}
 			else
 			{
-				CurrentTensorDescriptor.Sizes[Index] = CurrentTensorSize;
+				CurrentTensorDescriptor.Shape[Index] = CurrentTensorSize;
 				TensorNumberElements *= CurrentTensorSize;
 			}
 			++Index;
diff --git a/Engine/Plugins/Experimental/NNX/Source/NNXRuntimeORT/Private/NNXRuntimeORTUtils.h b/Engine/Plugins/Experimental/NNX/Source/NNXRuntimeORT/Private/NNXRuntimeORTUtils.h
--- a/Engine/Plugins/Experimental/NNX/Source/NNXRuntimeORT/Private/NNXRuntimeORTUtils.h
+++ b/Engine/Plugins/Experimental/NNX/Source/NNXRuntimeORT/Private/NNXRuntimeORTUtils.h
@@ -135,14 +135,14 @@ inline void BindTensorsToORT(
 		ONNXTensorElementDataType CurrentORTType = InTensorsORTType[Index];
 
 		TUniquePtr<int64_t[]> SizesInt64t;
-		SizesInt64t = MakeUnique<int64_t[]>(CurrentDescriptor.Dimension);
-		for (uint32 DimIndex = 0; DimIndex < CurrentDescriptor.Dimension; ++DimIndex)
+		SizesInt64t = MakeUnique<int64_t[]>(CurrentDescriptor.Shape.Num());
+		for (int32 DimIndex = 0; DimIndex < CurrentDescriptor.Shape.Num(); ++DimIndex)
 		{
-			SizesInt64t.Get()[DimIndex] = CurrentDescriptor.Sizes[DimIndex];
+			SizesInt64t.Get()[DimIndex] = CurrentDescriptor.Shape[DimIndex];
 		}
 
 		const uint64 ByteCount { InTensorsDescriptors[Index].DataSize };
-		const uint32 ArrayDimensions { CurrentDescriptor.Dimension };
+		const uint32 ArrayDimensions { (uint32)CurrentDescriptor.Shape.Num() };
 		OutOrtTensors.Emplace(
 			Ort::Value::CreateTensor(
 				*InAllocatorInfo,
diff --git a/Engine/Plugins/Experimental/NNX/Source/NNXRuntimeRDG/Private/NNXRuntimeDML.cpp b/Engine/Plugins/Experimental/NNX/Source/NNXRuntimeRDG/Private/NNXRuntimeDML.cpp
--- a/Engine/Plugins/Experimental/NNX/Source/NNXRuntimeRDG/Private/NNXRuntimeDML.cpp
+++ b/Engine/Plugins/Experimental/NNX/Source/NNXRuntimeRDG/Private/NNXRuntimeDML.cpp
@@ -216,8 +216,8 @@ protected:
 
 		BuffDesc.DataType = DmlDataType;
 		BuffDesc.Flags = DML_TENSOR_FLAG_NONE;
-		BuffDesc.DimensionCount = TensorDesc.Dimension;
-		BuffDesc.Sizes = TensorDesc.Sizes;
+		BuffDesc.DimensionCount = TensorDesc.Shape.Num();
+		BuffDesc.Sizes = TensorDesc.Shape.GetData();
 		// TODO: FIXME: Support strides in FMLTensorDesc
 		BuffDesc.Strides = nullptr;
 		BuffDesc.TotalTensorSizeInBytes = TensorDesc.DataSize;
diff --git a/Engine/Plugins/Experimental/NNX/Source/NNXRuntimeRDG/Private/NNXRuntimeHLSLConvTransposeOp.cpp b/Engine/Plugins/Experimental/NNX/Source/NNXRuntimeRDG/Private/NNXRuntimeHLSLConvTransposeOp.cpp
--- a/Engine/Plugins/Experimental/NNX/Source/NNXRuntimeRDG/Private/NNXRuntimeHLSLConvTransposeOp.cpp
+++ b/Engine/Plugins/Experimental/NNX/Source/NNXRuntimeRDG/Private/NNXRuntimeHLSLConvTransposeOp.cpp
@@ -55,9 +55,9 @@ namespace NNX
 			Weights = InputTensors[1];
 			Output = OutputTensors[0];
 
-			check(Input.Dimension > 2);
-			check(Weights.Dimension  == Input.Dimension);
-			check(Output.Dimension == Input.Dimension);
+			check(Input.Shape.Num() > 2);
+			check(Weights.Shape.Num()  == Input.Shape.Num());
+			check(Output.Shape.Num() == Input.Shape.Num());
 
 			if (InputTensors.Num() == 3) {
 				HasBias = true;
@@ -69,15 +69,15 @@ namespace NNX
 
 			auto MakeShape = [](const NNX::FMLTensorDesc& InputDesc) {
 				TArray<int32> Result;
-				for (int32 i = 0; i < (int32)InputDesc.Dimension; i++) {
-					Result.Add(InputDesc.Sizes[i]);
+				for (int32 i = 0; i < (int32)InputDesc.Shape.Num(); i++) {
+					Result.Add(InputDesc.Shape[i]);
 				}
 				return Result;
 			};
 			InputShape = MakeShape(Input);
 			WeightsShape = MakeShape(Weights);
 
-			NumDimensions = Input.Dimension - 2;
+			NumDimensions = Input.Shape.Num() - 2;
 
 			return true;
 		}
diff --git a/Engine/Plugins/Experimental/NNX/Source/NNXRuntimeRDG/Private/NNXRuntimeHLSLElementWiseBinaryOps.cpp b/Engine/Plugins/Experimental/NNX/Source/NNXRuntimeRDG/Private/NNXRuntimeHLSLElementWiseBinaryOps.cpp
--- a/Engine/Plugins/Experimental/NNX/Source/NNXRuntimeRDG/Private/NNXRuntimeHLSLElementWiseBinaryOps.cpp
+++ b/Engine/Plugins/Experimental/NNX/Source/NNXRuntimeRDG/Private/NNXRuntimeHLSLElementWiseBinaryOps.cpp
@@ -59,8 +59,8 @@ public:
 		Params->LHSInput = LHSInputSRV;
 		Params->RHSInput = RHSInputSRV;
 		Params->Output = OutputUAV;
-		FillTensorStrideForBroadcastShaderParameters(LHSInput, Output.Dimension, Params->TensorInfo, 0);
-		FillTensorStrideForBroadcastShaderParameters(RHSInput, Output.Dimension, Params->TensorInfo, 1);
+		FillTensorStrideForBroadcastShaderParameters(LHSInput, Output.Shape.Num(), Params->TensorInfo, 0);
+		FillTensorStrideForBroadcastShaderParameters(RHSInput, Output.Shape.Num(), Params->TensorInfo, 1);
 		FillTensorStrideShaderParameters(Output, Params->TensorInfo, 2);
 		Params->Num = Output.Num();
 		Params->ThreadCountX = ThreadGroupCount.X * FMLElementWiseBinaryCS::THREADGROUP_SIZE_X;
@@ -68,7 +68,7 @@ public:
 		FMLElementWiseBinaryCS::FPermutationDomain PermutationVector;
 
 		PermutationVector.Set<FMLElementWiseBinaryCS::FOperatorType>(OpType);
-		PermutationVector.Set<FMLElementWiseBinaryCS::FBinaryNumDimensions>(Output.Dimension);
+		PermutationVector.Set<FMLElementWiseBinaryCS::FBinaryNumDimensions>(Output.Shape.Num());
 
 		TShaderMapRef<FMLElementWiseBinaryCS> ComputeShader(GetGlobalShaderMap(GMaxRHIFeatureLevel), PermutationVector);
 		
diff --git a/Engine/Plugins/Experimental/NNX/Source/NNXRuntimeRDG/Private/NNXRuntimeHLSLElementWiseVariadicOps.cpp b/Engine/Plugins/Experimental/NNX/Source/NNXRuntimeRDG/Private/NNXRuntimeHLSLElementWiseVariadicOps.cpp
--- a/Engine/Plugins/Experimental/NNX/Source/NNXRuntimeRDG/Private/NNXRuntimeHLSLElementWiseVariadicOps.cpp
+++ b/Engine/Plugins/Experimental/NNX/Source/NNXRuntimeRDG/Private/NNXRuntimeHLSLElementWiseVariadicOps.cpp
@@ -40,18 +40,18 @@ namespace NNX
 		Params->Input2 = InputsSRV[2];
 		Params->Input3 = InputsSRV[3];
 		Params->Output = OutputUAV;
-		FillTensorStrideForBroadcastShaderParameters(InputDesc[0], OutputDesc.Dimension, Params->InputTensorInfo, 0);
+		FillTensorStrideForBroadcastShaderParameters(InputDesc[0], OutputDesc.Shape.Num(), Params->InputTensorInfo, 0);
 		if (InputBindings.Num() >= 2)
 		{
-			FillTensorStrideForBroadcastShaderParameters(InputDesc[1], OutputDesc.Dimension, Params->InputTensorInfo, 1);
+			FillTensorStrideForBroadcastShaderParameters(InputDesc[1], OutputDesc.Shape.Num(), Params->InputTensorInfo, 1);
 		}
 		if (InputBindings.Num() >= 3)
 		{
-			FillTensorStrideForBroadcastShaderParameters(InputDesc[2], OutputDesc.Dimension, Params->InputTensorInfo, 2);
+			FillTensorStrideForBroadcastShaderParameters(InputDesc[2], OutputDesc.Shape.Num(), Params->InputTensorInfo, 2);
 		}
 		if (InputBindings.Num() >= 4)
 		{
-			FillTensorStrideForBroadcastShaderParameters(InputDesc[3], OutputDesc.Dimension, Params->InputTensorInfo, 3);
+			FillTensorStrideForBroadcastShaderParameters(InputDesc[3], OutputDesc.Shape.Num(), Params->InputTensorInfo, 3);
 		}
 		FillTensorStrideShaderParameters(OutputDesc, Params->OutputTensorInfo, 0);
 		Params->Num = OutputDesc.Num();
@@ -65,7 +65,7 @@ namespace NNX
 		PermutationVector.Set<FMLElementWiseVariadicCS::FApplyScale>(Scale != 1.0f);
 		PermutationVector.Set<FMLElementWiseVariadicCS::FOutputAsInput>(OutputAsInput);
 		PermutationVector.Set<FMLElementWiseVariadicCS::FNumInput>(InputBindings.Num());
-		PermutationVector.Set<FMLElementWiseVariadicCS::FVariadicNumDimensions>(OutputDesc.Dimension);
+		PermutationVector.Set<FMLElementWiseVariadicCS::FVariadicNumDimensions>(OutputDesc.Shape.Num());
 
 		// Add the pass to RDG
 		TShaderMapRef<FMLElementWiseVariadicCS> ComputeShader(GetGlobalShaderMap(GMaxRHIFeatureLevel), PermutationVector);
diff --git a/Engine/Plugins/Experimental/NNX/Source/NNXRuntimeRDG/Private/NNXRuntimeHLSLGemmOp.cpp b/Engine/Plugins/Experimental/NNX/Source/NNXRuntimeRDG/Private/NNXRuntimeHLSLGemmOp.cpp
--- a/Engine/Plugins/Experimental/NNX/Source/NNXRuntimeRDG/Private/NNXRuntimeHLSLGemmOp.cpp
+++ b/Engine/Plugins/Experimental/NNX/Source/NNXRuntimeRDG/Private/NNXRuntimeHLSLGemmOp.cpp
@@ -56,11 +56,11 @@ namespace NNX
 			if (InputTensors.Num() == 3) InputC = InputTensors[2];
 			Output = OutputTensors[0];
 
-			check(InputA.Dimension == 2);
-			check(InputB.Dimension == 2);
-			check(InputC.Dimension < 3);
+			check(InputA.Shape.Num() == 2);
+			check(InputB.Shape.Num() == 2);
+			check(InputC.Shape.Num() < 3);
 
-			check(InputC.Dimension != 1 || InputC.Sizes[0] != 1); // TODO scalar version not supported yet
+			check(InputC.Shape.Num() != 1 || InputC.Shape[0] != 1); // TODO scalar version not supported yet
 
 			// C is treated as a scalar if there is no valid C, either width or height is zero or C dimension is 1x1
 			bIsCScalar = false; // InputTensors.Num() != 3 || InputC.Sizes[0] * InputC.Sizes[1] < 2;
diff --git a/Engine/Plugins/Experimental/NNX/Source/NNXRuntimeRDG/Private/NNXRuntimeHLSLHelper.cpp b/Engine/Plugins/Experimental/NNX/Source/NNXRuntimeRDG/Private/NNXRuntimeHLSLHelper.cpp
--- a/Engine/Plugins/Experimental/NNX/Source/NNXRuntimeRDG/Private/NNXRuntimeHLSLHelper.cpp
+++ b/Engine/Plugins/Experimental/NNX/Source/NNXRuntimeRDG/Private/NNXRuntimeHLSLHelper.cpp
@@ -6,42 +6,42 @@
 
 namespace NNX
 {
-	void FillTensorStrideShaderParameters(const FMLTensorDesc& TensorDesc, TStaticArray<FUintVector4, NXRT_TENSORSTRIDEINFO_MAX_NUM_DIMENSIONS, 16U>& OutShaderParam, uint32 Idx, uint32 TargetNumdimensionForBroadcast)
+	void FillTensorStrideShaderParameters(const FMLTensorDesc& TensorDesc, TStaticArray<FUintVector4, NXRT_TENSORSTRIDEINFO_MAX_NUM_DIMENSIONS, 16U>& OutShaderParam, int32 Idx, int32 TargetNumdimensionForBroadcast)
 	{
 		if (TargetNumdimensionForBroadcast == -1)
 		{
-			TargetNumdimensionForBroadcast = TensorDesc.Dimension;
+			TargetNumdimensionForBroadcast = TensorDesc.Shape.Num();
 		}
-		checkf(TargetNumdimensionForBroadcast >= TensorDesc.Dimension, TEXT("Can't broadcast tensor from rank %d to rank %d, should be inferior or equal."), TensorDesc.Dimension, TargetNumdimensionForBroadcast);
-		uint32 Offset = TargetNumdimensionForBroadcast - TensorDesc.Dimension;
+		checkf(TargetNumdimensionForBroadcast >= TensorDesc.Shape.Num(), TEXT("Can't broadcast tensor from rank %d to rank %d, should be inferior or equal."), TensorDesc.Shape.Num(), TargetNumdimensionForBroadcast);
+		int32 Offset = TargetNumdimensionForBroadcast - TensorDesc.Shape.Num();
 
 		static_assert(FMLTensorDesc::MaxTensorDimension <= NXRT_TENSORSTRIDEINFO_MAX_NUM_DIMENSIONS);
 		for (int32 i = 7; i >= 0; --i)
 		{
-			if ((uint32)i >= TargetNumdimensionForBroadcast || (uint32)i < Offset)
+			if (i >= TargetNumdimensionForBroadcast || i < Offset)
 			{
 				OutShaderParam[i][Idx] = 0;
 			}
-			else if ((uint32)i == TargetNumdimensionForBroadcast - 1)
+			else if (i == TargetNumdimensionForBroadcast - 1)
 			{
 				OutShaderParam[i][Idx] = 1;
 			}
 			else
 			{
-				OutShaderParam[i][Idx] = OutShaderParam[i + 1][Idx] * TensorDesc.Sizes[i + 1 - Offset];
+				OutShaderParam[i][Idx] = OutShaderParam[i + 1][Idx] * TensorDesc.Shape[i + 1 - Offset];
 			}
 		}
 	}
 
-	void FillTensorStrideForBroadcastShaderParameters(const FMLTensorDesc& TensorDesc, uint32 OutputNumdimension, TStaticArray<FUintVector4, NXRT_TENSORSTRIDEINFO_MAX_NUM_DIMENSIONS, 16U>& OutShaderParam, uint32 Idx)
+	void FillTensorStrideForBroadcastShaderParameters(const FMLTensorDesc& TensorDesc, int32 OutputNumdimension, TStaticArray<FUintVector4, NXRT_TENSORSTRIDEINFO_MAX_NUM_DIMENSIONS, 16U>& OutShaderParam, int32 Idx)
 	{
-		checkf(OutputNumdimension >= TensorDesc.Dimension, TEXT("Can't broadcast tensor from rank %d to rank %d, should be inferior or equal."), TensorDesc.Dimension, OutputNumdimension);
+		checkf(OutputNumdimension >= TensorDesc.Shape.Num(), TEXT("Can't broadcast tensor from rank %d to rank %d, should be inferior or equal."), TensorDesc.Shape.Num(), OutputNumdimension);
 		FillTensorStrideShaderParameters(TensorDesc, OutShaderParam, Idx, OutputNumdimension);
-		uint32 Offset = OutputNumdimension - TensorDesc.Dimension;
-		for (uint32 i = Offset; i < OutputNumdimension; ++i)
+		int32 Offset = OutputNumdimension - TensorDesc.Shape.Num();
+		for (int32 i = Offset; i < OutputNumdimension; ++i)
 		{
 			// the stride for broadcast dimension is kept as 0
-			if (TensorDesc.Sizes[i - Offset] == 1)
+			if (TensorDesc.Shape[i - Offset] == 1)
 			{
 				OutShaderParam[i][Idx] = 0;
 			}
diff --git a/Engine/Plugins/Experimental/NNX/Source/NNXRuntimeRDG/Private/NNXRuntimeHLSLHelper.h b/Engine/Plugins/Experimental/NNX/Source/NNXRuntimeRDG/Private/NNXRuntimeHLSLHelper.h
--- a/Engine/Plugins/Experimental/NNX/Source/NNXRuntimeRDG/Private/NNXRuntimeHLSLHelper.h
+++ b/Engine/Plugins/Experimental/NNX/Source/NNXRuntimeRDG/Private/NNXRuntimeHLSLHelper.h
@@ -8,7 +8,7 @@
 
 namespace NNX
 {
-	void FillTensorStrideShaderParameters(const FMLTensorDesc& TensorDesc, TStaticArray<FUintVector4, NXRT_TENSORSTRIDEINFO_MAX_NUM_DIMENSIONS, 16U>& OutShaderParam, uint32 Idx, uint32 TargetNumdimensionForBroadcast = -1);
-	void FillTensorStrideForBroadcastShaderParameters(const FMLTensorDesc& TensorDesc, uint32 OutputNumdimension, TStaticArray<FUintVector4, NXRT_TENSORSTRIDEINFO_MAX_NUM_DIMENSIONS, 16U>& OutShaderParam, uint32 Idx);
+	void FillTensorStrideShaderParameters(const FMLTensorDesc& TensorDesc, TStaticArray<FUintVector4, NXRT_TENSORSTRIDEINFO_MAX_NUM_DIMENSIONS, 16U>& OutShaderParam, int32 Idx, int32 TargetNumdimensionForBroadcast = -1);
+	void FillTensorStrideForBroadcastShaderParameters(const FMLTensorDesc& TensorDesc, int32 OutputNumdimension, TStaticArray<FUintVector4, NXRT_TENSORSTRIDEINFO_MAX_NUM_DIMENSIONS, 16U>& OutShaderParam, int32 Idx);
 	FIntVector ComputeElementWiseThreadGroups(uint32 ElementCount, uint32 GroupSizeX);
 } // NNX
\ No newline at end of file
diff --git a/Engine/Plugins/Experimental/NNX/Source/NNXRuntimeRDG/Private/NNXRuntimeHLSLMatMulOp.cpp b/Engine/Plugins/Experimental/NNX/Source/NNXRuntimeRDG/Private/NNXRuntimeHLSLMatMulOp.cpp
--- a/Engine/Plugins/Experimental/NNX/Source/NNXRuntimeRDG/Private/NNXRuntimeHLSLMatMulOp.cpp
+++ b/Engine/Plugins/Experimental/NNX/Source/NNXRuntimeRDG/Private/NNXRuntimeHLSLMatMulOp.cpp
@@ -48,7 +48,7 @@ namespace NNX
 		{
 			const EGemmAlgorithm Algorithm = EGemmAlgorithm::Simple32x32;
 
-			const int32 NumStackDimensions = FMath::Max((int32)FMath::Max(InputA.Dimension, InputB.Dimension) - 2, 0);
+			const int32 NumStackDimensions = FMath::Max(FMath::Max(InputA.Shape.Num(), InputB.Shape.Num()) - 2, 0);
 
 			// Set parameters
 			FMLGemmCS::FParameters* Parameters = GraphBuilder.AllocParameters<FMLGemmCS::FParameters>();
diff --git a/Engine/Plugins/Experimental/NNX/Source/NNXRuntimeRDG/Private/NNXRuntimeRDG.cpp b/Engine/Plugins/Experimental/NNX/Source/NNXRuntimeRDG/Private/NNXRuntimeRDG.cpp
--- a/Engine/Plugins/Experimental/NNX/Source/NNXRuntimeRDG/Private/NNXRuntimeRDG.cpp
+++ b/Engine/Plugins/Experimental/NNX/Source/NNXRuntimeRDG/Private/NNXRuntimeRDG.cpp
@@ -53,18 +53,12 @@ bool FMLInferenceModelRDG::LoadModel(UMLInferenceModel* InModel, FMLRuntimeForma
 
 		Tensor.Name = TensorDesc.Name;
 
-		Tensor.Dimension = TensorDesc.Shape.Dimension;
-		if (Tensor.Dimension > FMLTensorDesc::MaxTensorDimension)
+		if (TensorDesc.Shape.Num() > FMLTensorDesc::MaxTensorDimension)
 		{
 			// FIXME: Should we report error here?
-			Tensor.Dimension = FMLTensorDesc::MaxTensorDimension;
 		}
 
-		FMemory::Memset(Tensor.Sizes, 0, sizeof(Tensor.Sizes));
-		for (uint32 ShapeIdx = 0; ShapeIdx < Tensor.Dimension; ++ShapeIdx)
-		{
-			Tensor.Sizes[ShapeIdx] = TensorDesc.Shape.Sizes[ShapeIdx];
-		}
+		Tensor.Shape = TensorDesc.Shape;
 
 		// NOTE: Set first the DataType prior setting DataSize
 		Tensor.DataType = TensorDesc.DataType;
diff --git a/Engine/Plugins/Experimental/NNX/Source/NNXUtils/Private/NNXModelBuilderNNX.cpp b/Engine/Plugins/Experimental/NNX/Source/NNXUtils/Private/NNXModelBuilderNNX.cpp
--- a/Engine/Plugins/Experimental/NNX/Source/NNXUtils/Private/NNXModelBuilderNNX.cpp
+++ b/Engine/Plugins/Experimental/NNX/Source/NNXUtils/Private/NNXModelBuilderNNX.cpp
@@ -35,11 +35,11 @@ public:
 
 		Str << TEXT("[");
 		
-		for (uint32 Idx = 0; Idx < Tensor.Shape.Dimension; ++Idx)
+		for (int32 Idx = 0; Idx < Tensor.Shape.Num(); ++Idx)
 		{
-			Str << Tensor.Shape.Sizes[Idx];
+			Str << Tensor.Shape[Idx];
 
-			if (Idx + 1 < Tensor.Shape.Dimension)
+			if (Idx + 1 < Tensor.Shape.Num())
 			{
 				Str << TEXT(",");
 			}
diff --git a/Engine/Plugins/Experimental/NNX/Source/NNXUtils/Private/NNXModelBuilderONNX.cpp b/Engine/Plugins/Experimental/NNX/Source/NNXUtils/Private/NNXModelBuilderONNX.cpp
--- a/Engine/Plugins/Experimental/NNX/Source/NNXUtils/Private/NNXModelBuilderONNX.cpp
+++ b/Engine/Plugins/Experimental/NNX/Source/NNXUtils/Private/NNXModelBuilderONNX.cpp
@@ -294,7 +294,7 @@ NNXUTILS_API bool CreateONNXModelForOperator(const FString& OperatorName, TConst
 	for (int32 Idx = 0; Idx < InInputTensors.Num(); ++Idx)
 	{
 		const FMLTensorDesc& Desc = InInputTensors[Idx];
-		IMLModelBuilder::HTensor Tensor = Builder->AddTensor(Desc.Name, Desc.DataType, MakeArrayView((const int32*) Desc.Sizes, Desc.Dimension));
+		IMLModelBuilder::HTensor Tensor = Builder->AddTensor(Desc.Name, Desc.DataType, MakeArrayView((const int32*) Desc.Shape.GetData(), Desc.Shape.Num()));
 
 		InputTensors.Emplace(Tensor);
 		Builder->AddInput(Tensor);
@@ -305,7 +305,7 @@ NNXUTILS_API bool CreateONNXModelForOperator(const FString& OperatorName, TConst
 	for (int32 Idx = 0; Idx < InOutputTensors.Num(); ++Idx)
 	{
 		const FMLTensorDesc& Desc = InOutputTensors[Idx];
-		IMLModelBuilder::HTensor Tensor = Builder->AddTensor(Desc.Name, Desc.DataType, MakeArrayView((const int32*) Desc.Sizes, Desc.Dimension));
+		IMLModelBuilder::HTensor Tensor = Builder->AddTensor(Desc.Name, Desc.DataType, MakeArrayView((const int32*) Desc.Shape.GetData(), Desc.Shape.Num()));
 
 		OutputTensors.Emplace(Tensor);
 		Builder->AddOutput(Tensor);
@@ -340,4 +340,3 @@ NNXUTILS_API IMLModelBuilder* CreateONNXModelBuilder(int64 IrVersion, int64 Opse
 }
 
 } // namespace NNX
-
