- Added support for broadcasting for DirectML elementwise binary operators
- Added support for PRelu operator

#rnx

diff --git a/Engine/Plugins/Experimental/NNX/Source/NNXQA/Resources/NNXQATestDesc.json b/Engine/Plugins/Experimental/NNX/Source/NNXQA/Resources/NNXQATestDesc.json
--- a/Engine/Plugins/Experimental/NNX/Source/NNXQA/Resources/NNXQATestDesc.json
+++ b/Engine/Plugins/Experimental/NNX/Source/NNXQA/Resources/NNXQATestDesc.json
@@ -160,7 +160,6 @@
 		},
 		{
 			"category": "Elementwise.Binary",
-			"runtimes": [ { "name": "NNXRuntimeDml", "skip": true} ],
 			"targets": [
 				{	"target": "Add", "additional_datasets": [ "Elementwise.Slow.TwoInputs" ]},
 				{	"target": "And", "tags": [ "SmokeTest" ], "input_type": "boolean", "output_type": "boolean", "runtimes": [
diff --git a/Engine/Plugins/Experimental/NNX/Source/NNXRuntimeRDG/Private/NNXRuntimeDML.cpp b/Engine/Plugins/Experimental/NNX/Source/NNXRuntimeRDG/Private/NNXRuntimeDML.cpp
--- a/Engine/Plugins/Experimental/NNX/Source/NNXRuntimeRDG/Private/NNXRuntimeDML.cpp
+++ b/Engine/Plugins/Experimental/NNX/Source/NNXRuntimeRDG/Private/NNXRuntimeDML.cpp
@@ -154,6 +154,144 @@ OP_EW(SUBTRACT)
 #undef OP_EW
 #undef OP_AN
 
+namespace DmlUtil
+{
+	struct FTensorDesc
+	{
+		DML_BUFFER_TENSOR_DESC												BuffDesc;
+		DML_TENSOR_DESC														Desc;
+		TArray<uint32, TInlineAllocator<FMLTensorDesc::MaxTensorDimension>>	Sizes;
+		TArray<uint32, TInlineAllocator<FMLTensorDesc::MaxTensorDimension>>	Strides;
+	};
+
+	void SetTensorStrides(FTensorDesc& TensorDesc, const FMLTensorDesc& InputDesc)
+	{
+		uint32 CurrStride = 1;
+
+		TensorDesc.Strides.SetNum(InputDesc.Shape.Num());
+		
+		for (int32 i = InputDesc.Shape.Num() - 1; i >= 0; --i)
+		{
+			TensorDesc.Strides[i] = CurrStride;
+			CurrStride *= InputDesc.Shape[i];
+		}
+	}
+
+	void SetTensorSizesAndStridesForBroadcast(FTensorDesc& TensorDesc, const FMLTensorDesc& InputDesc, const FMLTensorDesc& TargetDesc)
+	{
+		static_assert(FMLTensorDesc::MaxTensorDimension <= 8);
+		
+		const uint32 TargetDimension = TargetDesc.Shape.Num() != -1 ? TargetDesc.Shape.Num() : InputDesc.Shape.Num();
+		checkf(TargetDesc.Shape.Num() >= InputDesc.Shape.Num(), TEXT("Can't broadcast tensor from rank %d to rank %d, should be inferior or equal."), InputDesc.Shape.Num(), TargetDimension);
+		
+		TensorDesc.Sizes.SetNum(TargetDimension);
+		TensorDesc.Strides.SetNum(TargetDimension);
+
+		const int32 DimensionOffset = int32(TargetDimension - InputDesc.Shape.Num());
+		
+		for (int32 i = 0; i < (int32) TargetDimension; ++i)
+		{
+			TensorDesc.Sizes[i] = i < DimensionOffset ? 1 : InputDesc.Shape[i - DimensionOffset];
+		}
+
+		uint32 CurrStride = 1;
+
+		for (int32 i = TargetDimension - 1; i >= 0; --i)
+		{
+			const bool bBroadcast = TensorDesc.Sizes[i] < TargetDesc.Shape[i];
+
+			TensorDesc.Strides[i] = bBroadcast ? 0 : CurrStride;
+			CurrStride *= TensorDesc.Sizes[i];
+
+			TensorDesc.Sizes[i] = TargetDesc.Shape[i];
+		}
+	}
+
+	inline bool IsSameShape(const FMLTensorDesc& Left, const FMLTensorDesc& Right)
+	{
+		if (Left.Shape.Num() != Right.Shape.Num())
+		{
+			return false;
+		}
+		
+		for (int32 Idx = 0; Idx < Left.Shape.Num(); ++Idx)
+		{
+			if (Left.Shape[Idx] != Right.Shape[Idx])
+			{
+				return false;
+			}
+		}
+
+		return true;
+	}
+
+	DML_TENSOR_DATA_TYPE GetTensorDataType(EMLTensorDataType DataType)
+	{
+		switch (DataType)
+		{
+		case EMLTensorDataType::Double:
+			return DML_TENSOR_DATA_TYPE_FLOAT64;
+
+		case EMLTensorDataType::Float:
+			return DML_TENSOR_DATA_TYPE_FLOAT32;
+
+		case EMLTensorDataType::Half:
+			return DML_TENSOR_DATA_TYPE_FLOAT16;
+
+		case EMLTensorDataType::UInt64:
+			return DML_TENSOR_DATA_TYPE_UINT64;
+
+		case EMLTensorDataType::UInt32:
+			return DML_TENSOR_DATA_TYPE_UINT32;
+
+		case EMLTensorDataType::UInt16:
+			return DML_TENSOR_DATA_TYPE_UINT16;
+
+		case EMLTensorDataType::UInt8:
+			return DML_TENSOR_DATA_TYPE_UINT8;
+
+		case EMLTensorDataType::Int64:
+			return DML_TENSOR_DATA_TYPE_INT64;
+
+		case EMLTensorDataType::Int32:
+			return DML_TENSOR_DATA_TYPE_INT32;
+
+		case EMLTensorDataType::Int16:
+			return DML_TENSOR_DATA_TYPE_INT16;
+
+		case EMLTensorDataType::Int8:
+			return DML_TENSOR_DATA_TYPE_INT8;
+
+		default:
+			return DML_TENSOR_DATA_TYPE_UNKNOWN;
+		}
+	}
+
+	inline uint64 CalculateBufferSize(const FTensorDesc& DmlTensor, const FMLTensorDesc& Desc)
+	{
+		uint64 ElemSizeInBytes = Desc.GetElemByteSize();
+		
+		if (ElemSizeInBytes == 0)
+		{
+			return 0;
+		}
+
+		uint64 MinSizeInBytes = 0;		
+		uint32 IndexOfLastElement = 0;
+
+		for (int32 i = 0; i < DmlTensor.Sizes.Num(); ++i)
+		{
+			IndexOfLastElement += (DmlTensor.Sizes[i] - 1) * DmlTensor.Strides[i];
+		}
+
+		MinSizeInBytes = (static_cast<uint64>(IndexOfLastElement) + 1) * ElemSizeInBytes;
+		
+		// Round up to the nearest 4 bytes.
+		MinSizeInBytes = (MinSizeInBytes + 3) & ~3ull;
+
+		return MinSizeInBytes;
+	}
+}
 
 //
 //
@@ -186,41 +324,95 @@ class FMLOperatorDml : public FMLOperatorRDG
 {
 public:
 
-	struct FDmlTensorDesc
-	{
-		DML_BUFFER_TENSOR_DESC	BuffDesc;
-		DML_TENSOR_DESC			Desc;
-	};
-
 	virtual bool Initialize(FDeviceContextDml* DevCtx, TArrayView<const FMLTensorDesc> InputTensors, TArrayView<const FMLTensorDesc> OutputTensors) = 0;
 
 	virtual void Dispatch(FRDGBuilder& GraphBuilder, TArrayView<const FMLTensorBinding> InInputBindings, TArrayView<const FMLTensorBinding> OutOutputBindings) = 0;
 
 protected:
 
-	bool InitDmlTensorDesc(FDmlTensorDesc& DmlTensorDesc, const FMLTensorDesc& TensorDesc)
+	bool InitDmlTensorDesc(DmlUtil::FTensorDesc& DmlTensorDesc, const FMLTensorDesc& TensorDesc)
 	{
-		DML_TENSOR_DATA_TYPE DmlDataType = GetDmlTensorDataType(TensorDesc.DataType);
+		DML_TENSOR_DATA_TYPE DmlDataType = DmlUtil::GetTensorDataType(TensorDesc.DataType);
 
 		if (DmlDataType == DML_TENSOR_DATA_TYPE_UNKNOWN)
 		{
 			DmlTensorDesc.BuffDesc = DML_BUFFER_TENSOR_DESC{};
 			DmlTensorDesc.Desc = DML_TENSOR_DESC{};
-			
+
 			return false;
 		}
 
-		auto& BuffDesc = DmlTensorDesc.BuffDesc;
+		DmlTensorDesc.Sizes = TensorDesc.Shape;
+		// TODO: Support tensor padding using strides defined in FMLTensorDesc
+		//DmlUtil::SetTensorStrides(DmlTensorDesc, TensorDesc.Strides);
 		
+		DML_BUFFER_TENSOR_DESC& BuffDesc = DmlTensorDesc.BuffDesc;
+
 		BuffDesc = DML_BUFFER_TENSOR_DESC{};
 
 		BuffDesc.DataType = DmlDataType;
 		BuffDesc.Flags = DML_TENSOR_FLAG_NONE;
 		BuffDesc.DimensionCount = TensorDesc.Shape.Num();
-		BuffDesc.Sizes = TensorDesc.Shape.GetData();
-		// TODO: FIXME: Support strides in FMLTensorDesc
+		BuffDesc.Sizes = DmlTensorDesc.Sizes.GetData();
 		BuffDesc.Strides = nullptr;
 		BuffDesc.TotalTensorSizeInBytes = TensorDesc.DataSize;
+
+		DmlTensorDesc.Desc = DML_TENSOR_DESC{ DML_TENSOR_TYPE_BUFFER, &DmlTensorDesc.BuffDesc };
+
+		return true;
+	}
+
+	bool InitDmlTensorDesc(DmlUtil::FTensorDesc& DmlTensorDesc, const FMLTensorDesc& TensorDesc, const FMLTensorDesc& BroadcastDesc)
+	{
+		DML_TENSOR_DATA_TYPE DmlDataType = DmlUtil::GetTensorDataType(TensorDesc.DataType);
+
+		if (DmlDataType == DML_TENSOR_DATA_TYPE_UNKNOWN)
+		{
+			DmlTensorDesc.BuffDesc = DML_BUFFER_TENSOR_DESC{};
+			DmlTensorDesc.Desc = DML_TENSOR_DESC{};
+
+			return false;
+		}
+
+		if (DmlUtil::IsSameShape(TensorDesc, BroadcastDesc))
+		{
+			DmlTensorDesc.Sizes = TensorDesc.Shape;
+			DmlUtil::SetTensorStrides(DmlTensorDesc, TensorDesc);
+		}
+		else if (TensorDesc.Shape.Num() > BroadcastDesc.Shape.Num())
+		{
+			return false;
+		}
+		else
+		{
+			DmlUtil::SetTensorSizesAndStridesForBroadcast(DmlTensorDesc, TensorDesc, BroadcastDesc);
+		}
+
+		//UE_LOG(LogNNX, Warning, TEXT("DmlTensorDesc:%d,%d,%d -> %d,%d,%d"),
+		//	TensorDesc.Sizes[0],
+		//	TensorDesc.Dimension > 1 ? TensorDesc.Sizes[1] : 0,
+		//	TensorDesc.Dimension > 2 ? TensorDesc.Sizes[2] : 0,
+		//	DmlTensorDesc.Sizes[0],
+		//	DmlTensorDesc.Dimension > 1 ? DmlTensorDesc.Sizes[1] : 0,
+		//	DmlTensorDesc.Dimension > 2 ? DmlTensorDesc.Sizes[2] : 0
+		//);
+
+		//UE_LOG(LogNNX, Warning, TEXT("DmlTensorStrides:%d,%d,%d"),
+		//	DmlTensorDesc.Strides[0], DmlTensorDesc.Strides[1], DmlTensorDesc.Strides[2]
+		//);
+
+		check(DmlTensorDesc.Strides.Num() == DmlTensorDesc.Sizes.Num());
+		
+		DML_BUFFER_TENSOR_DESC& BuffDesc = DmlTensorDesc.BuffDesc;
+		
+		BuffDesc = DML_BUFFER_TENSOR_DESC{};
+
+		BuffDesc.DataType = DmlDataType;
+		BuffDesc.Flags = DML_TENSOR_FLAG_NONE;
+		BuffDesc.DimensionCount = DmlTensorDesc.Sizes.Num();
+		BuffDesc.Sizes = DmlTensorDesc.Sizes.GetData();
+		BuffDesc.Strides = DmlTensorDesc.Strides.GetData();
+		BuffDesc.TotalTensorSizeInBytes = DmlUtil::CalculateBufferSize(DmlTensorDesc, TensorDesc);
 		
 		DmlTensorDesc.Desc = DML_TENSOR_DESC{ DML_TENSOR_TYPE_BUFFER, &DmlTensorDesc.BuffDesc };
 
@@ -297,39 +489,6 @@ protected:
 		return true;
 	}
 
-	DML_TENSOR_DATA_TYPE GetDmlTensorDataType(EMLTensorDataType DataType)
-	{
-		switch (DataType)
-		{
-			case EMLTensorDataType::Float:
-				return DML_TENSOR_DATA_TYPE_FLOAT32;
-
-			case EMLTensorDataType::Half:
-				return DML_TENSOR_DATA_TYPE_FLOAT16;
-
-			case EMLTensorDataType::UInt32:
-				return DML_TENSOR_DATA_TYPE_UINT32;
-			
-			case EMLTensorDataType::UInt16:
-				return DML_TENSOR_DATA_TYPE_UINT16;
-
-			case EMLTensorDataType::UInt8:
-				return DML_TENSOR_DATA_TYPE_UINT8;
-
-			case EMLTensorDataType::Int32:
-				return DML_TENSOR_DATA_TYPE_INT32;
-
-			case EMLTensorDataType::Int16:
-				return DML_TENSOR_DATA_TYPE_INT16;
-
-			case EMLTensorDataType::Int8:
-				return DML_TENSOR_DATA_TYPE_INT8;
-
-			default:
-				return DML_TENSOR_DATA_TYPE_UNKNOWN;
-		}
-	}
-
 	void ResetBindingTable(DML_BINDING_TABLE_DESC& Desc)
 	{
 		Desc.Dispatchable = CompiledOp;
@@ -398,7 +557,7 @@ public:
 		const FMLTensorDesc& OutputTensorDesc = OutputTensors[0];
 
 		// Initialize tensor descriptor (it's same for both input and output)
-		FDmlTensorDesc	DmlTensorDesc{};
+		DmlUtil::FTensorDesc	DmlTensorDesc{};
 
 		if (!InitDmlTensorDesc(DmlTensorDesc, InputTensorDesc))
 		{
@@ -427,20 +586,20 @@ public:
 private:
 
 	template<typename OpDesc>
-	void InitDmlOpDesc(OpDesc& Desc, FDmlTensorDesc& TensorDesc)
+	void InitDmlOpDesc(OpDesc& Desc, DmlUtil::FTensorDesc& TensorDesc)
 	{
 		Desc.InputTensor = &TensorDesc.Desc;
 		Desc.OutputTensor = &TensorDesc.Desc;
 	}
 
-	void InitDmlOpDesc(DML_ACTIVATION_SOFTPLUS_OPERATOR_DESC& Desc, FDmlTensorDesc& TensorDesc)
+	void InitDmlOpDesc(DML_ACTIVATION_SOFTPLUS_OPERATOR_DESC& Desc, DmlUtil::FTensorDesc& TensorDesc)
 	{
 		Desc.InputTensor = &TensorDesc.Desc;
 		Desc.OutputTensor = &TensorDesc.Desc;
 		Desc.Steepness = 1.0f;
 	}
 
-	void InitDmlOpDesc(DML_ACTIVATION_SCALED_ELU_OPERATOR_DESC& Desc, FDmlTensorDesc& TensorDesc)
+	void InitDmlOpDesc(DML_ACTIVATION_SCALED_ELU_OPERATOR_DESC& Desc, DmlUtil::FTensorDesc& TensorDesc)
 	{
 		Desc.InputTensor = &TensorDesc.Desc;
 		Desc.OutputTensor = &TensorDesc.Desc;
@@ -448,14 +607,14 @@ private:
 		Desc.Gamma = Gamma;
 	}
 
-	void InitDmlOpDesc(DML_ACTIVATION_ELU_OPERATOR_DESC& Desc, FDmlTensorDesc& TensorDesc)
+	void InitDmlOpDesc(DML_ACTIVATION_ELU_OPERATOR_DESC& Desc, DmlUtil::FTensorDesc& TensorDesc)
 	{
 		Desc.InputTensor = &TensorDesc.Desc;
 		Desc.OutputTensor = &TensorDesc.Desc;
 		Desc.Alpha = Alpha;
 	}
 
-	void InitDmlOpDesc(DML_ACTIVATION_HARD_SIGMOID_OPERATOR_DESC& Desc, FDmlTensorDesc& TensorDesc)
+	void InitDmlOpDesc(DML_ACTIVATION_HARD_SIGMOID_OPERATOR_DESC& Desc, DmlUtil::FTensorDesc& TensorDesc)
 	{
 		Desc.InputTensor = &TensorDesc.Desc;
 		Desc.OutputTensor = &TensorDesc.Desc;
@@ -463,7 +622,7 @@ private:
 		Desc.Beta = Beta;
 	}
 
-	void InitDmlOpDesc(DML_ACTIVATION_LEAKY_RELU_OPERATOR_DESC& Desc, FDmlTensorDesc& TensorDesc)
+	void InitDmlOpDesc(DML_ACTIVATION_LEAKY_RELU_OPERATOR_DESC& Desc, DmlUtil::FTensorDesc& TensorDesc)
 	{
 		Desc.InputTensor = &TensorDesc.Desc;
 		Desc.OutputTensor = &TensorDesc.Desc;
@@ -645,17 +804,17 @@ public:
 		const FMLTensorDesc& OutputTensorDesc = OutputTensors[0];
 
 		// Initialize tensor descriptors
-		FDmlTensorDesc	DmlInputATensorDesc{};
-		FDmlTensorDesc	DmlInputBTensorDesc{};
-		FDmlTensorDesc	DmlOutputTensorDesc{};
+		DmlUtil::FTensorDesc	DmlInputATensorDesc{};
+		DmlUtil::FTensorDesc	DmlInputBTensorDesc{};
+		DmlUtil::FTensorDesc	DmlOutputTensorDesc{};
 
-		if (!InitDmlTensorDesc(DmlInputATensorDesc, InputATensorDesc))
+		if (!InitDmlTensorDesc(DmlInputATensorDesc, InputATensorDesc, OutputTensorDesc))
 		{
 			UE_LOG(LogNNX, Warning, TEXT("Failed to initialize tensor(s) for DML inference"));
 			return false;
 		}
 
-		if (!InitDmlTensorDesc(DmlInputBTensorDesc, InputBTensorDesc))
+		if (!InitDmlTensorDesc(DmlInputBTensorDesc, InputBTensorDesc, OutputTensorDesc))
 		{
 			UE_LOG(LogNNX, Warning, TEXT("Failed to initialize tensor(s) for DML inference"));
 			return false;
@@ -688,20 +847,27 @@ public:
 private:
 
 	template<typename OpDesc>
-	void InitDmlOpDesc(OpDesc& Desc, FDmlTensorDesc& LHSTensor, FDmlTensorDesc& RHSTensor, FDmlTensorDesc& OutputTensor)
+	void InitDmlOpDesc(OpDesc& Desc, DmlUtil::FTensorDesc& LHSTensor, DmlUtil::FTensorDesc& RHSTensor, DmlUtil::FTensorDesc& OutputTensor)
 	{
 		Desc.ATensor = &LHSTensor.Desc;
 		Desc.BTensor = &RHSTensor.Desc;
 		Desc.OutputTensor = &OutputTensor.Desc;
 	}
 
-	void InitDmlOpDesc(DML_ELEMENT_WISE_POW_OPERATOR_DESC& Desc, FDmlTensorDesc& LHSTensor, FDmlTensorDesc& RHSTensor, FDmlTensorDesc& OutputTensor)
+	void InitDmlOpDesc(DML_ELEMENT_WISE_POW_OPERATOR_DESC& Desc, DmlUtil::FTensorDesc& LHSTensor, DmlUtil::FTensorDesc& RHSTensor, DmlUtil::FTensorDesc& OutputTensor)
 	{
 		Desc.InputTensor = &LHSTensor.Desc;
 		Desc.ExponentTensor = &RHSTensor.Desc;
 		Desc.OutputTensor = &OutputTensor.Desc;
 	}
 
+	void InitDmlOpDesc(DML_ACTIVATION_PARAMETERIZED_RELU_OPERATOR_DESC& Desc, DmlUtil::FTensorDesc& LHSTensor, DmlUtil::FTensorDesc& RHSTensor, DmlUtil::FTensorDesc& OutputTensor)
+	{
+		Desc.InputTensor = &LHSTensor.Desc;
+		Desc.SlopeTensor = &RHSTensor.Desc;
+		Desc.OutputTensor = &OutputTensor.Desc;
+	}
+
 public:
 
 	//
@@ -1045,7 +1211,7 @@ bool FMLRuntimeDml::RegisterElementWiseBinaryOperators()
 	// Mod
 	OP(DML_ELEMENT_WISE_MULTIPLY_OPERATOR_DESC, Mul);
 	// Or
-	// PRelu
+	OP(DML_ACTIVATION_PARAMETERIZED_RELU_OPERATOR_DESC, Prelu);
 	OP(DML_ELEMENT_WISE_POW_OPERATOR_DESC, Pow);
 	OP(DML_ELEMENT_WISE_SUBTRACT_OPERATOR_DESC, Sub);
 	// Xor
@@ -1131,43 +1297,6 @@ void FMLInferenceModelDml::AddDispatchOps_RenderThread(FRDGBuilder& GraphBuilder
 	}
 }
 
-////
-////
-////
-//uint64_t CalcBufferTensorSize(const DML_BUFFER_TENSOR_DESC& Desc)
-//{
-//	uint64_t Size = 0;
-//
-//	uint64_t ByteSize;
-//
-//	switch (Desc.DataType)
-//	{
-//	case DML_TENSOR_DATA_TYPE_FLOAT32:
-//		ByteSize = 4;
-//		break;
-//
-//	case DML_TENSOR_DATA_TYPE_FLOAT16:
-//		ByteSize = 2;
-//		break;
-//
-//	default:
-//		ByteSize = 0;
-//		break;
-//	}
-//
-//	if (ByteSize)
-//	{
-//		uint64_t ElemCount = 1;
-//
-//		for (uint32_t Dim = 0; Dim < Desc.DimensionCount; ++Dim)
-//			ElemCount *= Desc.Sizes[Desc.DimensionCount - 1 - Dim];
-//
-//		Size += ElemCount * ByteSize;
-//	}
-//
-//	return Size;
-//}
-
 //
 // Create operator
 //
