Added support for attributes for DirectML elementwise operators

#rnx

diff --git a/Engine/Plugins/Experimental/NNX/Source/NNXQA/Resources/NNXQATestDesc.json b/Engine/Plugins/Experimental/NNX/Source/NNXQA/Resources/NNXQATestDesc.json
--- a/Engine/Plugins/Experimental/NNX/Source/NNXQA/Resources/NNXQATestDesc.json
+++ b/Engine/Plugins/Experimental/NNX/Source/NNXQA/Resources/NNXQATestDesc.json
@@ -278,7 +278,6 @@
 			"datasets": [
 				{
 					"inputs": [ { "shape": [ 4 ] } ],
-					"runtimes": [ { "name": "NNXRuntimeDML",  "skip": true, "comment": "RDG DML does not support attribute yet" } ],
 					"comment": "for now there is no way to set attribute using the json test desc, this dataset is a special temporary case to test LeakyRelu with attribute not defaulted. See RunParametricTest() in NNQAParametricTest.cpp"
 				}
 			]
diff --git a/Engine/Plugins/Experimental/NNX/Source/NNXRuntimeRDG/Private/NNXRuntimeDML.cpp b/Engine/Plugins/Experimental/NNX/Source/NNXRuntimeRDG/Private/NNXRuntimeDML.cpp
--- a/Engine/Plugins/Experimental/NNX/Source/NNXRuntimeRDG/Private/NNXRuntimeDML.cpp
+++ b/Engine/Plugins/Experimental/NNX/Source/NNXRuntimeRDG/Private/NNXRuntimeDML.cpp
@@ -324,7 +324,7 @@ class FMLOperatorDml : public FMLOperatorRDG
 {
 public:
 
-	virtual bool Initialize(FDeviceContextDml* DevCtx, TArrayView<const FMLTensorDesc> InputTensors, TArrayView<const FMLTensorDesc> OutputTensors) = 0;
+	virtual bool Initialize(FDeviceContextDml* DevCtx, TArrayView<const FMLTensorDesc> InputTensors, TArrayView<const FMLTensorDesc> OutputTensors, const FMLAttributeMap& Attributes) = 0;
 
 	virtual void Dispatch(FRDGBuilder& GraphBuilder, TArrayView<const FMLTensorBinding> InInputBindings, TArrayView<const FMLTensorBinding> OutOutputBindings) = 0;
 
@@ -546,7 +546,7 @@ public:
 	//
 	//
 	//
-	virtual bool Initialize(FDeviceContextDml* InDevCtx, TArrayView<const FMLTensorDesc> InputTensors, TArrayView<const FMLTensorDesc> OutputTensors) override
+	virtual bool Initialize(FDeviceContextDml* InDevCtx, TArrayView<const FMLTensorDesc> InputTensors, TArrayView<const FMLTensorDesc> OutputTensors, const FMLAttributeMap& Attributes) override
 	{
 		// TODO: Setup attributes
 		Num = InputTensors[0].Num();
@@ -556,6 +556,10 @@ public:
 		const FMLTensorDesc& InputTensorDesc = InputTensors[0];
 		const FMLTensorDesc& OutputTensorDesc = OutputTensors[0];
 
+		Alpha = Attributes.GetOptionalFloat(TEXT("alpha"), Alpha);
+		Beta = Attributes.GetOptionalFloat(TEXT("beta"), Beta);
+		Gamma = Attributes.GetOptionalFloat(TEXT("gamma"), Gamma);
+
 		// Initialize tensor descriptor (it's same for both input and output)
 		DmlUtil::FTensorDesc	DmlTensorDesc{};
 
@@ -792,7 +796,7 @@ public:
 	//
 	//
 	//
-	virtual bool Initialize(FDeviceContextDml* InDevCtx, TArrayView<const FMLTensorDesc> InputTensors, TArrayView<const FMLTensorDesc> OutputTensors) override
+	virtual bool Initialize(FDeviceContextDml* InDevCtx, TArrayView<const FMLTensorDesc> InputTensors, TArrayView<const FMLTensorDesc> OutputTensors, const FMLAttributeMap& Attributes) override
 	{
 		// TODO: Setup attributes
 		Num = OutputTensors[0].Num();
@@ -977,7 +981,7 @@ protected:
 
 private:
 
-	FMLOperatorDml* OpCreate(const FString& Name, TArrayView<const FMLTensorDesc> InputTensorDesc, TArrayView<const FMLTensorDesc> OutputTensorDescs);
+	FMLOperatorDml* OpCreate(const FString& Name, TArrayView<const FMLTensorDesc> InputTensorDesc, TArrayView<const FMLTensorDesc> OutputTensorDescs, const FMLAttributeMap& Attributes);
 	
 	TArray<FMLOperatorDml*>		Operators;
 	FDeviceContextDml*			DevCtx;
@@ -1267,7 +1271,15 @@ bool FMLInferenceModelDml::Init(UMLInferenceModel* InModel, FDeviceContextDml* I
 		TArray<FMLTensorDesc> OpInputTensors = InputTensors;
 		TArray<FMLTensorDesc> OpOutputTensors = OutputTensors;
 
-		FMLOperatorDml* Op = OpCreate(TypeName, OpInputTensors, OpOutputTensors);
+		// Attributes
+		FMLAttributeMap Attributes;
+		
+		for (const FMLFormatAttributeDesc& Desc : Format.Operators[Idx].Attributes)
+		{
+			Attributes.SetAttribute(Desc.Name, Desc.Value);
+		}
+
+		FMLOperatorDml* Op = OpCreate(TypeName, OpInputTensors, OpOutputTensors, Attributes);
 
 		if (!Op)
 		{
@@ -1300,7 +1312,7 @@ void FMLInferenceModelDml::AddDispatchOps_RenderThread(FRDGBuilder& GraphBuilder
 //
 // Create operator
 //
-FMLOperatorDml* FMLInferenceModelDml::OpCreate(const FString& OpName, TArrayView<const FMLTensorDesc> InputTensorDescs, TArrayView<const FMLTensorDesc> OutputTensorDescs)
+FMLOperatorDml* FMLInferenceModelDml::OpCreate(const FString& OpName, TArrayView<const FMLTensorDesc> InputTensorDescs, TArrayView<const FMLTensorDesc> OutputTensorDescs, const FMLAttributeMap& Attributes)
 {
 	// TODO: Check if D3D12 device is valid
 
@@ -1314,7 +1326,7 @@ FMLOperatorDml* FMLInferenceModelDml::OpCreate(const FString& OpName, TArrayView
 
 	FMLOperatorDml*	Op = CreateFn();
 
-	if (!Op->Initialize(DevCtx, InputTensorDescs, OutputTensorDescs))
+	if (!Op->Initialize(DevCtx, InputTensorDescs, OutputTensorDescs, Attributes))
 	{
 		delete Op;
 
