Errors on uninitialized enums in NNXRuntimeFormat.h

diff --git a/Engine/Plugins/Experimental/NNX/Source/NNXCore/Public/NNXRuntimeFormat.h b/Engine/Plugins/Experimental/NNX/Source/NNXCore/Public/NNXRuntimeFormat.h
--- a/Engine/Plugins/Experimental/NNX/Source/NNXCore/Public/NNXRuntimeFormat.h
+++ b/Engine/Plugins/Experimental/NNX/Source/NNXCore/Public/NNXRuntimeFormat.h
@@ -76,16 +76,16 @@ struct FMLFormatTensorDesc
 	TArray<uint32> Shape;
 
 	UPROPERTY(VisibleAnywhere, Category = "Neural Network Inference")
-	EMLFormatTensorType	Type;
+	EMLFormatTensorType	Type = EMLFormatTensorType::None;
 
 	UPROPERTY(VisibleAnywhere, Category = "Neural Network Inference")
-	EMLTensorDataType	DataType;
+	EMLTensorDataType	DataType = EMLTensorDataType::None;
 
 	UPROPERTY(VisibleAnywhere, Category = "Neural Network Inference")
-	uint64	DataSize;
+	uint64	DataSize = 0;
 
 	UPROPERTY(VisibleAnywhere, Category = "Neural Network Inference")
-	uint64	DataOffset;
+	uint64	DataOffset = 0;
 };
 
 /// NNX Runtime format
