NNI - Implement first draft of gather op

diff --git a/Engine/Plugins/Experimental/NNX/Source/NNXHlslShaders/Private/NNXGatherCS.cpp b/Engine/Plugins/Experimental/NNX/Source/NNXHlslShaders/Private/NNXGatherCS.cpp
new file mode 100644
--- /dev/null
+++ b/Engine/Plugins/Experimental/NNX/Source/NNXHlslShaders/Private/NNXGatherCS.cpp
@@ -0,0 +1,80 @@
+// Copyright Epic Games, Inc. All Rights Reserved.
+
+#include "NNXGatherCS.h"
+
+namespace UE::NNI::HlslShaders::Internal
+{
+	template <>
+	void TGatherCS<float, int32>::ModifyCompilationEnvironment(const FGlobalShaderPermutationParameters& InParameters, FShaderCompilerEnvironment& OutEnvironment)
+	{
+		FGlobalShader::ModifyCompilationEnvironment(InParameters, OutEnvironment);
+		OutEnvironment.SetDefine(TEXT("DATA_TYPE"), 0);
+		OutEnvironment.SetDefine(TEXT("INDICES_TYPE"), 0);
+		OutEnvironment.SetDefine(TEXT("MAX_NUM_DIMENSIONS"), FGatherConstants::MAX_NUM_DIMENSIONS);
+		OutEnvironment.SetDefine(TEXT("NUM_GROUP_THREADS"), FGatherConstants::NUM_GROUP_THREADS);
+	}
+
+	template <typename DataElementType, typename IndicesElementType>
+	void TGatherCS<DataElementType, IndicesElementType>::FillInParameters(int32 Axis, const NNX::FMLTensorDesc& Data, const NNX::FMLTensorDesc& Indices, FParameters& Parameters)
+	{
+		Parameters.Axis = Axis;
+
+		TArray<int32> OutputShape;
+		OutputShape.SetNum(Indices.Shape.Num() + (Data.Shape.Num() - 1));
+		int32 index = 0;
+		Parameters.OutputSize = 1;
+		for (int32 i = 0; i < Axis; i++, index++)
+		{
+			OutputShape[index] = Data.Shape[i];
+			Parameters.OutputSize *= OutputShape[index];
+		}
+		for (int32 i = 0; i < Indices.Shape.Num(); i++, index++)
+		{
+			OutputShape[index] = Indices.Shape[i];
+			Parameters.OutputSize *= OutputShape[index];
+		}
+		for (int32 i = (Axis + 1); i < Data.Shape.Num(); i++, index++)
+		{
+			OutputShape[index] = Data.Shape[i];
+			Parameters.OutputSize *= OutputShape[index];
+		}
+
+		Parameters.NumDataDimensions = Data.Shape.Num();
+		Parameters.NumIndicesDimensions = Indices.Shape.Num();
+
+		int32 DataStride = 1;
+		for (int32 i = Data.Shape.Num() - 1; i >= 0; i--)
+		{
+			Parameters.DataStride_IndicesStride_OutputStride[i].X = DataStride;
+			Parameters.OneDivDataStride_OneDivIndicesStride_OneDivOutputStride[i].X = 1.0 / (float)DataStride;
+			DataStride *= Data.Shape[i];
+		}
+
+		int32 IndicesStride = 1;
+		for (int32 i = Indices.Shape.Num() - 1; i >= 0; i--)
+		{
+			Parameters.DataStride_IndicesStride_OutputStride[i].Y = IndicesStride;
+			Parameters.OneDivDataStride_OneDivIndicesStride_OneDivOutputStride[i].Y = 1.0 / (float)IndicesStride;
+			IndicesStride *= Indices.Shape[i];
+		}
+
+		int32 OutputStride = 1;
+		for (int32 i = OutputShape.Num() - 1; i >= 0; i--)
+		{
+			Parameters.DataStride_IndicesStride_OutputStride[i].Z = OutputStride;
+			Parameters.OneDivDataStride_OneDivIndicesStride_OneDivOutputStride[i].Z = 1.0 / (float)OutputStride;
+			OutputStride *= OutputShape[i];
+		}
+	}
+
+	template <typename DataElementType, typename IndicesElementType>
+	FIntVector TGatherCS<DataElementType, IndicesElementType>::GetGroupCount(const FParameters& Parameters)
+	{
+		return FIntVector(FMath::DivideAndRoundUp(Parameters.OutputSize, FGatherConstants::NUM_GROUP_THREADS), 1, 1);
+	}
+
+	template class TGatherCS<float, int32>;
+	typedef TGatherCS<float, int32> FGatherFloatInt32CS;
+	IMPLEMENT_SHADER_TYPE(template<>, FGatherFloatInt32CS, TEXT("/NNX/GatherOp.usf"), TEXT("main"), SF_Compute);
+
+} // UE::NNI::HlslShaders::Internal
\ No newline at end of file
diff --git a/Engine/Plugins/Experimental/NNX/Source/NNXHlslShaders/Public/NNXGatherCS.h b/Engine/Plugins/Experimental/NNX/Source/NNXHlslShaders/Public/NNXGatherCS.h
new file mode 100644
--- /dev/null
+++ b/Engine/Plugins/Experimental/NNX/Source/NNXHlslShaders/Public/NNXGatherCS.h
@@ -0,0 +1,47 @@
+// Copyright Epic Games, Inc. All Rights Reserved.
+
+#pragma once
+
+#include "GlobalShader.h"
+#include "ShaderParameterUtils.h"
+#include "RenderGraphUtils.h"
+#include "NNXTypes.h"
+
+namespace UE::NNI::HlslShaders::Internal
+{
+	class FGatherConstants
+	{
+	public:
+		static const int32 MAX_NUM_DIMENSIONS{ 8 };
+		static const int32 NUM_GROUP_THREADS{ 256 };
+	};
+
+	template <typename DataElementType, typename IndicesElementType>
+	class NNXHLSLSHADERS_API TGatherCS : public FGlobalShader
+	{
+		DECLARE_GLOBAL_SHADER(TGatherCS);
+		SHADER_USE_PARAMETER_STRUCT(TGatherCS, FGlobalShader)
+
+		class FGatherNumOutputDimensions : SHADER_PERMUTATION_RANGE_INT("NUM_OUTPUT_DIMENSIONS", 1, FGatherConstants::MAX_NUM_DIMENSIONS);
+		using FPermutationDomain = TShaderPermutationDomain<FGatherNumOutputDimensions>;
+
+	public:
+		BEGIN_SHADER_PARAMETER_STRUCT(FParameters, )
+			SHADER_PARAMETER(int32, Axis)
+			SHADER_PARAMETER(int32, OutputSize)
+			SHADER_PARAMETER(int32, NumDataDimensions)
+			SHADER_PARAMETER(int32, NumIndicesDimensions)
+			SHADER_PARAMETER_ARRAY(FIntVector4, DataStride_IndicesStride_OutputStride, [FGatherConstants::MAX_NUM_DIMENSIONS])
+			SHADER_PARAMETER_ARRAY(FVector4f, OneDivDataStride_OneDivIndicesStride_OneDivOutputStride, [FGatherConstants::MAX_NUM_DIMENSIONS])
+			SHADER_PARAMETER_RDG_BUFFER_SRV(Buffer<DataElementType>, Data)
+			SHADER_PARAMETER_RDG_BUFFER_SRV(Buffer<IndicesElementType>, Indices)
+			SHADER_PARAMETER_RDG_BUFFER_UAV(RWBuffer<DataElementType>, Output)
+		END_SHADER_PARAMETER_STRUCT()
+
+		static void ModifyCompilationEnvironment(const FGlobalShaderPermutationParameters& InParameters, FShaderCompilerEnvironment& OutEnvironment);
+
+		static void FillInParameters(int32 Axis, const NNX::FMLTensorDesc& Data, const NNX::FMLTensorDesc& Indices, FParameters& Parameters);
+
+		static FIntVector GetGroupCount(const FParameters& Parameters);
+	};
+} // UE::NNI::HlslShaders::Internal
\ No newline at end of file
diff --git a/Engine/Plugins/Experimental/NNX/Source/NNXHlslShaders/Shaders/GatherOp.usf b/Engine/Plugins/Experimental/NNX/Source/NNXHlslShaders/Shaders/GatherOp.usf
new file mode 100644
--- /dev/null
+++ b/Engine/Plugins/Experimental/NNX/Source/NNXHlslShaders/Shaders/GatherOp.usf
@@ -0,0 +1,86 @@
+// Copyright Epic Games, Inc. All Rights Reserved.
+
+#include "/Engine/Public/Platform.ush"
+
+#if DATA_TYPE == 0
+#define DATA_BUFFER_TYPE float
+#else
+#define DATA_BUFFER_TYPE float
+#endif
+
+#if INDICES_TYPE == 0
+#define INDICES_WORK_TYPE int
+#define INDICES_BUFFER_TYPE int
+#define INDICES_BUFFER_TO_WORK_TYPE(x) x
+#else
+#define INDICES_WORK_TYPE int
+#define INDICES_BUFFER_TYPE int
+#define INDICES_BUFFER_TO_WORK_TYPE(x) x
+#endif
+
+int Axis;
+int OutputSize;
+int NumDataDimensions;
+int NumIndicesDimensions;
+int4 DataStride_IndicesStride_OutputStride[MAX_NUM_DIMENSIONS];
+float4 OneDivDataStride_OneDivIndicesStride_OneDivOutputStride[MAX_NUM_DIMENSIONS];
+
+Buffer<DATA_BUFFER_TYPE> Data;
+Buffer<INDICES_BUFFER_TYPE> Indices;
+RWBuffer<DATA_BUFFER_TYPE> Output;
+
+[numthreads(NUM_GROUP_THREADS, 1, 1)]
+void main(in const uint3 DispatchThreadID : SV_DispatchThreadID)
+{
+	if (DispatchThreadID.x >= OutputSize)
+	{
+		return;
+	}
+
+	// Compute the output index per dimenion
+	INDICES_WORK_TYPE OutputIndices[NUM_OUTPUT_DIMENSIONS];
+	INDICES_WORK_TYPE DispatchThreadIndex = (INDICES_WORK_TYPE)DispatchThreadID.x;
+	UNROLL
+	for (int i = 0; i < NUM_OUTPUT_DIMENSIONS; i++)
+	{
+		OutputIndices[i] = (INDICES_WORK_TYPE)(OneDivDataStride_OneDivIndicesStride_OneDivOutputStride[i].z * (float)DispatchThreadIndex);
+		DispatchThreadIndex -= OutputIndices[i] * (INDICES_WORK_TYPE)DataStride_IndicesStride_OutputStride[i].z;
+	}
+
+	// Split the output indices into data and indices indices 
+	INDICES_WORK_TYPE DataIndices[NUM_OUTPUT_DIMENSIONS];
+	INDICES_WORK_TYPE IndicesIndices[NUM_OUTPUT_DIMENSIONS];
+	int AxisPlusNumIndicesDimensions = Axis + NumIndicesDimensions;
+	for (int i = 0; i < Axis; i++)
+	{
+		DataIndices[i] = OutputIndices[i];
+	}
+	for (int i = 0; i < NumIndicesDimensions; i++)
+	{
+		IndicesIndices[i] = OutputIndices[Axis + i];
+	}
+	for (int i = Axis + 1; i < NumDataDimensions; i++)
+	{
+		DataIndices[i] = OutputIndices[AxisPlusNumIndicesDimensions + i];
+	}
+
+	// Compute the flattened indices index
+	INDICES_WORK_TYPE IndicesIndex = 0;
+	for (int i = 0; i < NumIndicesDimensions; i++)
+	{
+		IndicesIndex += IndicesIndices[i] * DataStride_IndicesStride_OutputStride[i].y;
+	}
+
+	// Look up and assign the missing data index
+	DataIndices[Axis] = INDICES_BUFFER_TO_WORK_TYPE(Indices[IndicesIndex]);
+
+	// Compute the flattened data index
+	INDICES_WORK_TYPE DataIndex = 0;
+	for (int i = 0; i < NumDataDimensions; i++)
+	{
+		DataIndex += DataIndices[i] * DataStride_IndicesStride_OutputStride[i].x;
+	}
+
+	// Write the result
+	Output[DispatchThreadID.x] = Data[DataIndex];
+}
\ No newline at end of file
diff --git a/Engine/Plugins/Experimental/NNX/Source/NNXRuntimeRDG/Private/NNXRuntimeHLSLGatherOp.cpp b/Engine/Plugins/Experimental/NNX/Source/NNXRuntimeRDG/Private/NNXRuntimeHLSLGatherOp.cpp
new file mode 100644
--- /dev/null
+++ b/Engine/Plugins/Experimental/NNX/Source/NNXRuntimeRDG/Private/NNXRuntimeHLSLGatherOp.cpp
@@ -0,0 +1,92 @@
+// Copyright Epic Games, Inc. All Rights Reserved.
+
+#include "NNXRuntimeHLSLGatherOp.h"
+#include "NNXRuntimeHLSLHelper.h"
+#include "NNXGatherCS.h"
+
+namespace UE::NNI::RuntimeRDG::Hlsl::Private
+{
+	DECLARE_GPU_STAT_NAMED(FNNIOperatorGather, TEXT("NNI.Operator.Hlsl.Gather"));
+
+	template<typename DataElementType, typename IndicesElementType>
+	using TGatherCS = typename UE::NNI::HlslShaders::Internal::TGatherCS<DataElementType, IndicesElementType>;
+	using FGatherConstants = UE::NNI::HlslShaders::Internal::FGatherConstants;
+
+	template <typename DataElementType, typename IndicesElementType>
+	class FOperatorGather : public NNX::FMLOperatorHlsl
+	{
+	public:
+			
+		virtual ~FOperatorGather() = default;
+
+		FOperatorGather() {}
+
+	private:
+
+		int32 Axis = 0;
+		NNX::FMLTensorDesc Data = {};
+		NNX::FMLTensorDesc Indices = {};
+		NNX::FMLTensorDesc Output = {};
+
+	public:
+
+		virtual bool Initialize(TArrayView<const NNX::FMLTensorDesc> InputTensors, TArrayView<const NNX::FMLTensorDesc> OutputTensors, const FMLAttributeMap& Attributes) override
+		{
+			check(InputTensors.Num() == 2)
+			check(OutputTensors.Num() == 1)
+			check(Output.Shape.Num() <= FGatherConstants::MAX_NUM_DIMENSIONS)
+			check(InputTensors[0].Shape.Num() > 0)
+			check(InputTensors[1].Shape.Num() > 0)
+			check(InputTensors[0].Shape.Num() + (InputTensors[1].Shape.Num() - 1) <= FGatherConstants::MAX_NUM_DIMENSIONS)
+
+			Axis = Attributes.GetOptionalFloat(TEXT("axis"), Axis);
+			check(Axis < InputTensors[0].Shape.Num())
+			check(Axis >= -InputTensors[0].Shape.Num())
+			Axis = Axis >= 0 ? Axis : InputTensors[0].Shape.Num() + Axis;
+
+			Data = InputTensors[0];
+			Indices = InputTensors[1];
+			Output = OutputTensors[0];
+
+			return true;
+		}
+
+		virtual void Dispatch(FRDGBuilder& GraphBuilder, TArrayView<const NNX::FMLTensorBinding> InInputBindings, TArrayView<const NNX::FMLTensorBinding> OutOutputBindings) override
+		{
+			// Set parameters
+			typename TGatherCS<DataElementType, IndicesElementType>::FParameters* Parameters = GraphBuilder.AllocParameters<typename TGatherCS<DataElementType, IndicesElementType>::FParameters>();
+			TGatherCS<DataElementType, IndicesElementType>::FillInParameters(Axis, Data, Indices, *Parameters);
+			Parameters->Data = GraphBuilder.CreateSRV(FRDGBufferSRVDesc(InInputBindings[0].Buffer, PF_R32_FLOAT));
+			Parameters->Indices = GraphBuilder.CreateSRV(FRDGBufferSRVDesc(InInputBindings[1].Buffer, PF_R32_FLOAT));
+			Parameters->Output = GraphBuilder.CreateUAV(FRDGBufferUAVDesc(OutOutputBindings[0].Buffer, PF_R32_FLOAT));
+
+			typename TGatherCS<DataElementType, IndicesElementType>::FPermutationDomain PermutationVector;
+			PermutationVector.template Set<typename TGatherCS<DataElementType, IndicesElementType>::FGatherNumOutputDimensions>(Output.Shape.Num());
+			TShaderMapRef<TGatherCS<DataElementType, IndicesElementType>> ComputeShader(GetGlobalShaderMap(GMaxRHIFeatureLevel), PermutationVector);
+
+			FIntVector ThreadGroupCount = TGatherCS<DataElementType, IndicesElementType>::GetGroupCount(*Parameters);
+
+			RDG_EVENT_SCOPE(GraphBuilder, "NNI.Operator.Hlsl.Gather");
+			RDG_GPU_STAT_SCOPE(GraphBuilder, FNNIOperatorGather);
+
+			FComputeShaderUtils::AddPass(
+				GraphBuilder,
+				RDG_EVENT_NAME("NNI.Operator.Hlsl.Gather.Dispatch"),
+				ERDGPassFlags::Compute | ERDGPassFlags::NeverCull,
+				ComputeShader,
+				Parameters,
+				ThreadGroupCount);
+		}
+	};
+
+	NNX::FMLOperatorHlsl* CreateGatherOperator()
+	{
+		return new FOperatorGather<float, int32>();
+	}
+
+	bool RegisterGatherOperator(NNX::FMLOperatorRegistryHlsl& Registry)
+	{
+		Registry.OpAdd(TEXT("Gather"), CreateGatherOperator);
+		return true;
+	}
+} // UE::NNI::RuntimeRDG::Hlsl::Private
\ No newline at end of file
diff --git a/Engine/Plugins/Experimental/NNX/Source/NNXRuntimeRDG/Private/NNXRuntimeHLSLGatherOp.h b/Engine/Plugins/Experimental/NNX/Source/NNXRuntimeRDG/Private/NNXRuntimeHLSLGatherOp.h
new file mode 100644
--- /dev/null
+++ b/Engine/Plugins/Experimental/NNX/Source/NNXRuntimeRDG/Private/NNXRuntimeHLSLGatherOp.h
@@ -0,0 +1,11 @@
+// Copyright Epic Games, Inc. All Rights Reserved.
+
+#pragma once
+
+#include "NNXRuntimeRDG.h"
+#include "NNXRuntimeHLSLOp.h"
+
+namespace UE::NNI::RuntimeRDG::Hlsl::Private
+{
+	bool RegisterGatherOperator(NNX::FMLOperatorRegistryHlsl& Registry);
+} // UE::NNI::RuntimeRDG::Hlsl::Private
