Add attribute support to NNXQA

diff --git a/Engine/Plugins/Experimental/NNX/Source/NNXCore/Public/NNXTypes.h b/Engine/Plugins/Experimental/NNX/Source/NNXCore/Public/NNXTypes.h
--- a/Engine/Plugins/Experimental/NNX/Source/NNXCore/Public/NNXTypes.h
+++ b/Engine/Plugins/Experimental/NNX/Source/NNXCore/Public/NNXTypes.h
@@ -39,6 +39,21 @@ enum class EMLAttributeDataType : uint8
     //TODO add more AttributeDataType support
 };
 
+/**
+ * @return EMLAttributeDataType from the string passed in
+ */
+inline void LexFromString(EMLAttributeDataType& OutValue, const TCHAR* StringVal)
+{
+	int64 EnumVal = StaticEnum<EMLAttributeDataType>()->GetValueByName(StringVal);
+	if (EnumVal == INDEX_NONE)
+	{
+		OutValue = EMLAttributeDataType::None;
+		ensureMsgf(false, TEXT("EMLAttributeDataType LexFromString didn't have a match for '%s'"), StringVal);
+	}
+
+	OutValue = (EMLAttributeDataType)EnumVal;
+};
+
 USTRUCT()
 struct FMLAttributeValue
 {
diff --git a/Engine/Plugins/Experimental/NNX/Source/NNXHlslShaders/Private/NNXGemmCS.cpp b/Engine/Plugins/Experimental/NNX/Source/NNXHlslShaders/Private/NNXGemmCS.cpp
--- a/Engine/Plugins/Experimental/NNX/Source/NNXHlslShaders/Private/NNXGemmCS.cpp
+++ b/Engine/Plugins/Experimental/NNX/Source/NNXHlslShaders/Private/NNXGemmCS.cpp
@@ -15,11 +15,9 @@ namespace
 		TArray<uint32> StackShapeB;
 		TArray<uint32> StackStrideA;
 		TArray<uint32> StackStrideB;
-		uint32 CWidth = 0;
-		uint32 CHeight = 0;
 	};
 
-	static GemmMatrixParameters GetMatrixParameters(const NNX::FMLTensorDesc &InputA, const NNX::FMLTensorDesc &InputB, const NNX::FMLTensorDesc &InputC)
+	static GemmMatrixParameters GetMatrixParametersMatMul(const NNX::FMLTensorDesc &InputA, const NNX::FMLTensorDesc &InputB)
 	{
 		check(InputA.Shape.Num() != 0);
 		check(InputB.Shape.Num() != 0);
@@ -46,7 +44,6 @@ namespace
 
 		const bool IsVectorA = InputA.Shape.Num() == 1;
 		const bool IsVectorB = InputB.Shape.Num() == 1;
-		const bool IsVectorC = InputC.Shape.Num() == 1;
 
 		Result.M = IsVectorA ? 1 : InputA.Shape[InputA.Shape.Num() - 2];
 		Result.N = IsVectorB ? 1 : InputB.Shape[InputB.Shape.Num() - 1];
@@ -74,9 +71,6 @@ namespace
 			Result.StackStrideB[i] = Result.StackStrideB[i + 1] * Result.StackShapeB[i + 1];
 		}
 
-		Result.CWidth = InputC.Shape.Num() == 0 ? 0 : InputC.Shape[InputC.Shape.Num() - 1];
-		Result.CHeight = InputC.Shape.Num() == 0 ? 0 : IsVectorC ? 1 : InputC.Shape[InputC.Shape.Num() - 2];
-
 		return Result;
 	}
 } // namespace
@@ -90,26 +84,32 @@ void FMLGemmCS::ModifyCompilationEnvironment(const FGlobalShaderPermutationParam
 void FMLGemmCS::FillInParameters(float Alpha, float Beta, int32 TransA, int32 TransB, const NNX::FMLTensorDesc &InputA, const NNX::FMLTensorDesc &InputB,
 		const NNX::FMLTensorDesc &InputC, float CScalar, FMLGemmCS::FParameters& Parameters)
 {
-	GemmMatrixParameters Params = GetMatrixParameters(InputA, InputB, InputC);
+	check(InputA.Shape.Num() == 2);
+	check(InputB.Shape.Num() == 2);
+
+	uint32 M = TransA != 0 ? InputA.Shape[1] : InputA.Shape[0];
+	uint32 K = TransA != 0 ? InputA.Shape[0] : InputA.Shape[1];
+	uint32 N = TransB != 0 ? InputB.Shape[0] : InputB.Shape[1];
+	check(K == (TransB != 0 ? InputB.Shape[1] : InputB.Shape[0]));
 
 	Parameters.Alpha = Alpha;
 	Parameters.Beta = Beta;
 	Parameters.TransA = TransA;
 	Parameters.TransB = TransB;
-	Parameters.M = Params.M;
-	Parameters.N = Params.N;
-	Parameters.K = Params.K;
-	Parameters.MxK = Params.M * Params.K;
-	Parameters.KxN = Params.K * Params.N;
-	Parameters.MxN = Params.M * Params.N;
-	Parameters.CWidth = Params.CWidth;
-	Parameters.CHeight = Params.CHeight;
+	Parameters.M = M;
+	Parameters.N = N;
+	Parameters.K = K;
+	Parameters.MxK = M * K;
+	Parameters.KxN = K * N;
+	Parameters.MxN = M * N;
+	Parameters.CWidth = InputC.Shape.Num() == 0 ? 0 : InputC.Shape[InputC.Shape.Num() - 1];;
+	Parameters.CHeight = InputC.Shape.Num() == 0 ? 0 : InputC.Shape.Num() == 1 ? 1 : InputC.Shape[InputC.Shape.Num() - 2];
 	Parameters.CScalar = CScalar;
 }
 
 void FMLGemmCS::FillInParametersMatMul(const NNX::FMLTensorDesc &InputA, const NNX::FMLTensorDesc &InputB, FMLGemmCS::FParameters& Parameters)
 {
-	GemmMatrixParameters Params = GetMatrixParameters(InputA, InputB, {});
+	GemmMatrixParameters Params = GetMatrixParametersMatMul(InputA, InputB);
 
 	Parameters.Alpha = 1.0f;
 	Parameters.Beta = 1.0f;
@@ -121,8 +121,8 @@ void FMLGemmCS::FillInParametersMatMul(const NNX::FMLTensorDesc &InputA, const N
 	Parameters.MxK = Params.M * Params.K;
 	Parameters.KxN = Params.K * Params.N;
 	Parameters.MxN = Params.M * Params.N;
-	Parameters.CWidth = Params.CWidth;
-	Parameters.CHeight = Params.CHeight;
+	Parameters.CWidth = 0;
+	Parameters.CHeight = 0;
 	Parameters.CScalar = 0.0f;
 
 	for (int32 i = 0; i < Params.StackShapeA.Num(); i++)
diff --git a/Engine/Plugins/Experimental/NNX/Source/NNXQA/Private/NNXQAJsonUtils.h b/Engine/Plugins/Experimental/NNX/Source/NNXQA/Private/NNXQAJsonUtils.h
--- a/Engine/Plugins/Experimental/NNX/Source/NNXQA/Private/NNXQAJsonUtils.h
+++ b/Engine/Plugins/Experimental/NNX/Source/NNXQA/Private/NNXQAJsonUtils.h
@@ -7,6 +7,7 @@
 #include "Serialization/JsonSerializerMacros.h"
 #include "Dom/JsonObject.h"
 #include "Misc/FileHelper.h"
+#include "NNXTypes.h"
 
 namespace NNX 
 {
@@ -61,6 +62,56 @@ namespace Json
 		END_JSON_SERIALIZER
 	};
 
+	struct FTestAttribute : FJsonSerializable
+	{
+		FString Name;
+		FMLAttributeValue Value;
+
+		BEGIN_JSON_SERIALIZER
+			JSON_SERIALIZE("name", Name);
+			
+			// Since we prefer the attribute value to be in readable form in the Json and
+			// we only serialize it, we construct it at runtime by manually serialize it
+			// note: default type is float
+
+			FString TypeStr;
+			Serializer.Serialize(TEXT("type"), TypeStr);
+
+			EMLAttributeDataType Type = EMLAttributeDataType::Float;
+			if (!TypeStr.IsEmpty()) LexFromString(Type, *TypeStr);
+
+			switch (Type)
+			{
+				case EMLAttributeDataType::Float:
+				{
+					float TmpValue;
+					JSON_SERIALIZE("value", TmpValue);
+					Value = FMLAttributeValue(TmpValue);
+					break;
+				}
+
+				case EMLAttributeDataType::Int32:
+				{
+					int TmpValue;
+					JSON_SERIALIZE("value", TmpValue);
+					Value = FMLAttributeValue(TmpValue);
+					break;
+				}
+				default:
+					check(Type == EMLAttributeDataType::None);
+			}
+		END_JSON_SERIALIZER
+	};
+
+	struct FTestAttributeMap : FJsonSerializable
+	{
+		TArray<FTestAttribute> Attributes;
+
+		BEGIN_JSON_SERIALIZER
+			JSON_SERIALIZE_ARRAY_SERIALIZABLE("attributes", Attributes, FTestAttribute);
+		END_JSON_SERIALIZER
+	};
+
 	struct FTestConfigTarget : FJsonSerializable
 	{
 		FString Target;
@@ -122,14 +173,49 @@ namespace Json
 		END_JSON_SERIALIZER
 	};
 
+	struct FTestAttributeSet : FJsonSerializable
+	{
+		FString Name;
+		TArray<FTestAttributeMap> AttributeMaps;
+		TArray<FString> MultiplyWithAttributeSets;
+
+		BEGIN_JSON_SERIALIZER
+			JSON_SERIALIZE("name", Name);
+			JSON_SERIALIZE_ARRAY_SERIALIZABLE("attribute_maps", AttributeMaps, FTestAttributeMap);
+			JSON_SERIALIZE_ARRAY("multiply_with_attribute_sets", MultiplyWithAttributeSets);
+		END_JSON_SERIALIZER
+	};
+
+	template<typename T>
+	bool TrySerializeArray(TArray<T> &OutArray, const FString &FieldName, const FJsonObject &GlobalJsonObject)
+	{
+		const TArray<TSharedPtr<FJsonValue>>* JsonArrayValue;
+		if (GlobalJsonObject.TryGetArrayField(FieldName, JsonArrayValue))
+		{
+			for (TSharedPtr<FJsonValue> JsonValue : *JsonArrayValue)
+			{
+				const TSharedPtr<FJsonObject>* JsonObject = nullptr;
+				if (!JsonValue->TryGetObject(JsonObject))
+				{
+					return false;
+				}
+				OutArray.Emplace_GetRef().FromJson(*JsonObject);
+			}
+		}
+
+		return true;
+	}
+
 	bool LoadTestDescriptionFromJson(const FString& FullPath,
 		TArray<FTestCategory>& ModelTestCategories,
 		TArray<FTestCategory>& OperatorCategories,
-		TArray<FTestConfigInputOutputSet>& InputOutputSets)
+		TArray<FTestConfigInputOutputSet>& InputOutputSets,
+		TArray<FTestAttributeSet>& AttributeSets)
 	{
 		ModelTestCategories.Empty();
 		OperatorCategories.Empty();
 		InputOutputSets.Empty();
+		AttributeSets.Empty();
 
 		FString JsonContent;
 		FFileHelper::LoadFileToString(JsonContent, *FullPath);
@@ -146,46 +232,13 @@ namespace Json
 			return false;
 		}
 
-		const TArray<TSharedPtr<FJsonValue>>* JsonArrayValue;
-		if (GlobalJsonObject->TryGetArrayField(TEXT("model_test_categories"), JsonArrayValue))
-		{
-			for (TSharedPtr<FJsonValue> JsonValue : *JsonArrayValue)
-			{
-				const TSharedPtr<FJsonObject>* JsonObject = nullptr;
-				if (!JsonValue->TryGetObject(JsonObject))
-				{
-					return false;
-				}
-				ModelTestCategories.Emplace_GetRef().FromJson(*JsonObject);
-			}
-		}
-
-		if (GlobalJsonObject->TryGetArrayField(TEXT("operator_test_categories"), JsonArrayValue))
-		{
-			for (TSharedPtr<FJsonValue> JsonValue : *JsonArrayValue)
-			{
-				const TSharedPtr<FJsonObject>* JsonObject = nullptr;
-				if (!JsonValue->TryGetObject(JsonObject))
-				{
-					return false;
-				}
-				OperatorCategories.Emplace_GetRef().FromJson(*JsonObject);
-			}
-		}
+		bool bSuccess = false;
+		bSuccess |= TrySerializeArray(ModelTestCategories, TEXT("model_test_categories"), *GlobalJsonObject);
+		bSuccess |= TrySerializeArray(OperatorCategories, TEXT("operator_test_categories"), *GlobalJsonObject);
+		bSuccess |= TrySerializeArray(InputOutputSets, TEXT("input_output_sets"), *GlobalJsonObject);
+		bSuccess |= TrySerializeArray(AttributeSets, TEXT("attribute_sets"), *GlobalJsonObject);
 
-		if (GlobalJsonObject->TryGetArrayField(TEXT("input_output_sets"), JsonArrayValue))
-		{
-			for (TSharedPtr<FJsonValue> JsonValue : *JsonArrayValue)
-			{
-				const TSharedPtr<FJsonObject>* JsonObject = nullptr;
-				if (!JsonValue->TryGetObject(JsonObject))
-				{
-					return false;
-				}
-				InputOutputSets.Emplace_GetRef().FromJson(*JsonObject);
-			}
-		}
-		return true;
+		return bSuccess;
 	}
 
 } // namespace Json
diff --git a/Engine/Plugins/Experimental/NNX/Source/NNXQA/Private/NNXQAParametricTest.cpp b/Engine/Plugins/Experimental/NNX/Source/NNXQA/Private/NNXQAParametricTest.cpp
--- a/Engine/Plugins/Experimental/NNX/Source/NNXQA/Private/NNXQAParametricTest.cpp
+++ b/Engine/Plugins/Experimental/NNX/Source/NNXQA/Private/NNXQAParametricTest.cpp
@@ -48,15 +48,16 @@ namespace Test
 			TArray<Json::FTestCategory> ModelTestCategories;
 			TArray<Json::FTestCategory> OperatorTestCategories;
 			TArray<Json::FTestConfigInputOutputSet> InputOutputSets;
+			TArray<Json::FTestAttributeSet> AttributeSets;
 			
-			if (!Json::LoadTestDescriptionFromJson(FullPath, ModelTestCategories, OperatorTestCategories, InputOutputSets))
+			if (!Json::LoadTestDescriptionFromJson(FullPath, ModelTestCategories, OperatorTestCategories, InputOutputSets, AttributeSets))
 			{
 				return false;
 			}
 
 			const FString NNXBaseTestPath(TEXT("System.Engine.MachineLearning.NNX"));
-			AddTestFromCategory(NNXBaseTestPath + TEXT(".Model."), ModelTestCategories, InputOutputSets);
-			AddTestFromCategory(NNXBaseTestPath + TEXT(".Operator."), OperatorTestCategories, InputOutputSets);
+			AddTestFromCategory(NNXBaseTestPath + TEXT(".Model."), ModelTestCategories, InputOutputSets, AttributeSets);
+			AddTestFromCategory(NNXBaseTestPath + TEXT(".Operator."), OperatorTestCategories, InputOutputSets, AttributeSets);
 
 			return true;
 		}
@@ -109,7 +110,7 @@ namespace Test
 				TestSetup.RelativeErrorPercent = TestTarget.RelativeError;
 			}
 
-			ApplyRuntimesConfig(TestSetup, TestTarget.Runtimes);;
+			ApplyRuntimesConfig(TestSetup, TestTarget.Runtimes);
 
 			//Tags
 			TestSetup.Tags = TestTarget.Tags;
@@ -167,6 +168,14 @@ namespace Test
 			}
 		}
 
+		static void ApplyAttributeSetConfig(FTests::FTestSetup& TestSetup, const Json::FTestAttributeMap& AttributeMap)
+		{
+			for (const Json::FTestAttribute &Attribute : AttributeMap.Attributes)
+			{
+				TestSetup.AttributeMap.SetAttribute(Attribute.Name, Attribute.Value);
+			}
+		}
+
 		static FString GetTestSuffix(const Json::FTestConfigDataset& Dataset)
 		{
 			//Build TestSuffix "<inputshape0>_...=><outputshape0>_..."
@@ -195,6 +204,35 @@ namespace Test
 			return TestSuffix;
 		}
 
+		static FString GetAttributeMapSuffix(const Json::FTestAttributeMap &AttributeMap)
+		{
+			if (AttributeMap.Attributes.IsEmpty()) return TEXT("");
+
+			auto AttributeToString = [] (const FMLAttributeValue &Value) -> FString {
+				EMLAttributeDataType Type = Value.GetType();
+				
+				switch(Type)
+				{
+					case EMLAttributeDataType::Float: return FString::SanitizeFloat(Value.AsFloat());
+					case EMLAttributeDataType::Int32: return FString::FromInt(Value.AsInt32());
+				}
+				return TEXT("-");
+			};
+			
+			FStringBuilderBase Str;
+			Str  << TEXT(".");
+
+			TArray<FString> AttributeStrings;
+			for (const Json::FTestAttribute &Attribute : AttributeMap.Attributes)
+			{
+				AttributeStrings.Push(Attribute.Name + "=" + AttributeToString(Attribute.Value));
+			}
+
+			Str  << FString::Join(AttributeStrings, TEXT(","));
+
+			return Str.ToString();
+		}
+
 		static bool IsSubstringFoundInArray(const TArray<FString>& Names, const FString& SubString)
 		{
 			for (auto&& Name : Names)
@@ -207,7 +245,8 @@ namespace Test
 			return false;
 		}
 		
-		void AddTestFromCategory(const FString& BaseTestPath, const TArray<Json::FTestCategory>& TestCategories, const TArray<Json::FTestConfigInputOutputSet> InputOutputSets)
+		void AddTestFromCategory(const FString& BaseTestPath, const TArray<Json::FTestCategory>& TestCategories, const TArray<Json::FTestConfigInputOutputSet>& InputOutputSets,
+				const TArray<Json::FTestAttributeSet>& AttributeSets)
 		{
 			for (const Json::FTestCategory& TestCategory : TestCategories)
 			{
@@ -252,13 +291,76 @@ namespace Test
 									continue;
 								}
 
-								FTests::FTestSetup& Test = AddTest(TestCategoryPath, TestBaseName, TEXT(".") + GetTestSuffix(Dataset));
-								
-								ApplyRuntimesConfig(Test, TestCategory.Runtimes);
-								ApplyRuntimesConfig(Test, InputOutputSet.Runtimes);
-								ApplyTargetConfig(Test, TestTarget);
-								ApplyDatasetConfig(Test, Dataset, InputTypeFromTarget, OutputTypeFromTarget);
-								Test.IsModelTest = bIsModelCategory;
+								bool bAtLeastAnAttributeTestWasWadded = false;
+								for (const Json::FTestAttributeSet &AttributeSet : AttributeSets)
+								{
+									// TODO may be split by '.' and (partially) match parts
+									if (!AttributeSet.Name.Contains(InputOutputSet.Name))
+									{
+										continue;
+									}
+
+									for (const Json::FTestAttributeMap &AttributeMap : AttributeSet.AttributeMaps)
+									{
+										bool bAtLeastAnotherAttributeSetAdded = false;
+										for (const FString &OtherAttributeSetName : AttributeSet.MultiplyWithAttributeSets)
+										{
+											for (const Json::FTestAttributeSet &OtherAttributeSet : AttributeSets)
+											{
+												if (OtherAttributeSet.Name != OtherAttributeSetName)
+												{
+													continue;
+												}
+
+												for (const Json::FTestAttributeMap &OtherAttributeMap : OtherAttributeSet.AttributeMaps)
+												{
+													FTests::FTestSetup& Test = AddTest(TestCategoryPath, TestBaseName, TEXT(".") + GetTestSuffix(Dataset) + GetAttributeMapSuffix(AttributeMap) + GetAttributeMapSuffix(OtherAttributeMap));
+										
+													ApplyRuntimesConfig(Test, TestCategory.Runtimes);
+													ApplyRuntimesConfig(Test, InputOutputSet.Runtimes);
+													ApplyTargetConfig(Test, TestTarget);
+													ApplyDatasetConfig(Test, Dataset, InputTypeFromTarget, OutputTypeFromTarget);
+
+													ApplyAttributeSetConfig(Test, AttributeMap);
+													ApplyAttributeSetConfig(Test, OtherAttributeMap);
+
+													Test.IsModelTest = bIsModelCategory;
+
+													bAtLeastAnAttributeTestWasWadded = true;
+
+													bAtLeastAnotherAttributeSetAdded = true;
+												}
+											}
+										}
+
+										if (!bAtLeastAnotherAttributeSetAdded)
+										{
+											FTests::FTestSetup& Test = AddTest(TestCategoryPath, TestBaseName, TEXT(".") + GetTestSuffix(Dataset) + GetAttributeMapSuffix(AttributeMap));
+											
+											ApplyRuntimesConfig(Test, TestCategory.Runtimes);
+											ApplyRuntimesConfig(Test, InputOutputSet.Runtimes);
+											ApplyTargetConfig(Test, TestTarget);
+											ApplyDatasetConfig(Test, Dataset, InputTypeFromTarget, OutputTypeFromTarget);
+
+											ApplyAttributeSetConfig(Test, AttributeMap);
+
+											Test.IsModelTest = bIsModelCategory;
+
+											bAtLeastAnAttributeTestWasWadded = true;
+										}
+									}
+								}
+
+								if (!bAtLeastAnAttributeTestWasWadded)
+								{
+									FTests::FTestSetup& Test = AddTest(TestCategoryPath, TestBaseName, TEXT(".") + GetTestSuffix(Dataset));
+									
+									ApplyRuntimesConfig(Test, TestCategory.Runtimes);
+									ApplyRuntimesConfig(Test, InputOutputSet.Runtimes);
+									ApplyTargetConfig(Test, TestTarget);
+									ApplyDatasetConfig(Test, Dataset, InputTypeFromTarget, OutputTypeFromTarget);
+									Test.IsModelTest = bIsModelCategory;
+								}
 							}
 							bAtLeastATestWasAdded = true;
 						}
@@ -303,14 +405,8 @@ namespace Test
 		}
 		else
 		{
-			//TODO add support for attribute in json schema
-			FMLAttributeMap AttributeMap;
-			if (TestSetup.TestName == TEXT("System.Engine.MachineLearning.NNX.Operator.Elementwise.Unary.LeakyRelu.[4]=>[4]"))
-			{
-				AttributeMap.SetAttribute("alpha", FMLAttributeValue(0.5f));
-			}
 			// Operator test, create model in memory
-			if (!CreateONNXModelForOperator(TestSetup.TargetName, TestSetup.Inputs, TestSetup.Outputs, AttributeMap, ModelData))
+			if (!CreateONNXModelForOperator(TestSetup.TargetName, TestSetup.Inputs, TestSetup.Outputs, TestSetup.AttributeMap, ModelData))
 			{
 				UE_LOG(LogNNX, Error, TEXT("Failed to create model for test '%s'. Test ABORTED!"), *TestSetup.TargetName);
 				return false;
diff --git a/Engine/Plugins/Experimental/NNX/Source/NNXQA/Private/NNXQAUtils.h b/Engine/Plugins/Experimental/NNX/Source/NNXQA/Private/NNXQAUtils.h
--- a/Engine/Plugins/Experimental/NNX/Source/NNXQA/Private/NNXQAUtils.h
+++ b/Engine/Plugins/Experimental/NNX/Source/NNXQA/Private/NNXQAUtils.h
@@ -34,6 +34,7 @@ namespace Test
 			TMap<FString, float> RelativeErrorPercentForRuntime;
 			TArray<FMLTensorDesc> Inputs;
 			TArray<FMLTensorDesc> Outputs;
+			FMLAttributeMap AttributeMap;
 			TArray<FString> Tags;
 			TArray<FString> AutomationExcludedRuntime;
 			TArray<FString> AutomationExcludedPlatform;
diff --git a/Engine/Plugins/Experimental/NNX/Source/NNXQA/Resources/NNXQATestDesc.json b/Engine/Plugins/Experimental/NNX/Source/NNXQA/Resources/NNXQATestDesc.json
--- a/Engine/Plugins/Experimental/NNX/Source/NNXQA/Resources/NNXQATestDesc.json
+++ b/Engine/Plugins/Experimental/NNX/Source/NNXQA/Resources/NNXQATestDesc.json
@@ -116,7 +116,7 @@
 				{	"target": "IsNan", "skip": true, "comment": "TODO add flexibility in test data setup + bool tensors" },
 				{	"target": "HardSigmoid",         "comment": "TODO test with not-default attribute" },
 				{	"target": "HardSwish", "runtimes": [ { "name": "NNXRuntimeDml", "skip": true, "comment": "Not implemented atm on NNXRuntimeDML."} ]},
-				{	"target": "LeakyRelu", "additional_datasets": [ "Elementwise.LeakyRelu.HackedAttributeTest" ], "comment": "TODO test with not-default attribute" },
+				{	"target": "LeakyRelu", "additional_datasets": [ "Elementwise.LeakyRelu.AttributeTest" ] },
 				{	"target": "Log", "runtimes": [
 						{ "name": "NNXRuntimeORTDml", "relative_error": 0.15},
 						{ "name": "NNXRuntimeDml",    "relative_error": 0.15},
@@ -274,11 +274,10 @@
 			]
 		},
 		{
-			"name": "Elementwise.LeakyRelu.HackedAttributeTest",
+			"name": "Elementwise.LeakyRelu.AttributeTest",
 			"datasets": [
 				{
-					"inputs": [ { "shape": [ 4 ] } ],
-					"comment": "for now there is no way to set attribute using the json test desc, this dataset is a special temporary case to test LeakyRelu with attribute not defaulted. See RunParametricTest() in NNQAParametricTest.cpp"
+					"inputs": [ { "shape": [ 4 ] } ]
 				}
 			]
 		},
@@ -370,14 +369,46 @@
 			]
 		},
 		{
-			"name": "Gemm",
+			"name": "Gemm.TwoInputs",
 			"datasets": [
 				{ "inputs": [ { "shape": [ 2, 2 ] }, { "shape": [ 2, 2 ] } ], "outputs": [ { "shape": [ 2, 2 ] } ]},
-				{ "inputs": [ { "shape": [ 3, 4 ] }, { "shape": [ 4, 5 ] } ], "outputs": [ { "shape": [ 3, 5 ] } ]},
+				{ "inputs": [ { "shape": [ 3, 4 ] }, { "shape": [ 4, 5 ] } ], "outputs": [ { "shape": [ 3, 5 ] } ]}
+			]
+		},
+		{
+			"name": "Gemm.ThreeInputs",
+			"datasets": [
 				{ "inputs": [ { "shape": [ 3, 4 ] }, { "shape": [ 4, 5 ] }, { "shape": [ 3, 5 ] } ], "outputs": [ { "shape": [ 3, 5 ] } ]},
 				{ "inputs": [ { "shape": [ 1, 1 ] }, { "shape": [ 1, 5 ] }, { "shape": [ 5 ] } ], "outputs": [ { "shape": [ 1, 5 ] } ]}
 			]
 		},
+		{
+			"name": "Gemm.TransposedA",
+			"datasets": [
+				{ "inputs": [ { "shape": [ 2, 2 ] }, { "shape": [ 2, 2 ] } ], "outputs": [ { "shape": [ 2, 2 ] } ]},
+				{ "inputs": [ { "shape": [ 4, 3 ] }, { "shape": [ 4, 5 ] } ], "outputs": [ { "shape": [ 3, 5 ] } ]},
+				{ "inputs": [ { "shape": [ 4, 3 ] }, { "shape": [ 4, 5 ] }, { "shape": [ 3, 5 ] } ], "outputs": [ { "shape": [ 3, 5 ] } ]},
+				{ "inputs": [ { "shape": [ 1, 1 ] }, { "shape": [ 1, 5 ] }, { "shape": [ 5 ] } ], "outputs": [ { "shape": [ 1, 5 ] } ]}
+			]
+		},
+		{
+			"name": "Gemm.TransposedB",
+			"datasets": [
+				{ "inputs": [ { "shape": [ 2, 2 ] }, { "shape": [ 2, 2 ] } ], "outputs": [ { "shape": [ 2, 2 ] } ]},
+				{ "inputs": [ { "shape": [ 3, 4 ] }, { "shape": [ 5, 4 ] } ], "outputs": [ { "shape": [ 3, 5 ] } ]},
+				{ "inputs": [ { "shape": [ 3, 4 ] }, { "shape": [ 5, 4 ] }, { "shape": [ 3, 5 ] } ], "outputs": [ { "shape": [ 3, 5 ] } ]},
+				{ "inputs": [ { "shape": [ 1, 1 ] }, { "shape": [ 5, 1 ] }, { "shape": [ 5 ] } ], "outputs": [ { "shape": [ 1, 5 ] } ]}
+			]
+		},
+		{
+			"name": "Gemm.AllTransposed",
+			"datasets": [
+				{ "inputs": [ { "shape": [ 2, 2 ] }, { "shape": [ 2, 2 ] } ], "outputs": [ { "shape": [ 2, 2 ] } ]},
+				{ "inputs": [ { "shape": [ 4, 3 ] }, { "shape": [ 5, 4 ] } ], "outputs": [ { "shape": [ 3, 5 ] } ]},
+				{ "inputs": [ { "shape": [ 4, 3 ] }, { "shape": [ 5, 4 ] }, { "shape": [ 3, 5 ] } ], "outputs": [ { "shape": [ 3, 5 ] } ]},
+				{ "inputs": [ { "shape": [ 1, 1 ] }, { "shape": [ 5, 1 ] }, { "shape": [ 5 ] } ], "outputs": [ { "shape": [ 1, 5 ] } ]}
+			]
+		},
 		{
 			"name": "ConvTranspose",
 			"datasets": [
@@ -415,6 +446,66 @@
 				{ "inputs": [ { "shape": [ 2, 1, 3, 5, 4 ] }, { "shape": [ 4, 3 ] } ], "outputs": [ { "shape": [ 2, 1, 3, 5, 3 ] } ]}
 			]
 		}
+	],
+
+	"attribute_sets": [
+		{
+			"name": "Elementwise.LeakyRelu.AttributeTest",
+			"attribute_maps": [
+				{ "attributes": [ { "name": "alpha", "value": 0.1 } ] },
+				{ "attributes": [ { "name": "alpha", "value": 0.5 } ] },
+				{ "attributes": [ { "name": "alpha", "value": 1.0 } ] }
+			]
+		},
+		{
+			"name": "Gemm.TwoInputs.Alpha",
+			"attribute_maps": [
+				{ "attributes": [ { "name": "alpha", "value": 0.0 } ] },
+				{ "attributes": [ { "name": "alpha", "value": 0.5 } ] },
+				{ "attributes": [ { "name": "alpha", "value": 1.0 } ] },
+				{ "attributes": [ { "name": "alpha", "value": 5.0 } ] }
+			]
+		},
+		{
+			"name": "Gemm.ThreeInputs.Beta",
+			"multiply_with_attribute_sets": [ "Gemm.TwoInputs.Alpha" ],
+			"attribute_maps": [
+				{ "attributes": [ { "name": "beta", "value": 0.0 } ] },
+				{ "attributes": [ { "name": "beta", "value": 0.5 } ] },
+				{ "attributes": [ { "name": "beta", "value": 1.0 } ] },
+				{ "attributes": [ { "name": "beta", "value": 5.0 } ] }
+			]
+		},
+		{
+			"name": "Gemm.TwoInputs.DefaultTranspose",
+			"attribute_maps": [
+				{ "attributes": [ { "name": "transA", "value": 0, "type": "Int32" }, { "name": "transB", "value": 0, "type": "Int32" } ] }
+			]
+		},
+		{
+			"name": "Gemm.ThreeInputs.DefaultTranspose",
+			"attribute_maps": [
+				{ "attributes": [ { "name": "transA", "value": 0, "type": "Int32" }, { "name": "transB", "value": 0, "type": "Int32" } ] }
+			]
+		},
+		{
+			"name": "Gemm.TransposedA",
+			"attribute_maps": [
+				{ "attributes": [ { "name": "transA", "value": 1, "type": "Int32" } ] }
+			]
+		},
+		{
+			"name": "Gemm.TransposedB",
+			"attribute_maps": [
+				{ "attributes": [ { "name": "transB", "value": 1, "type": "Int32" } ] }
+			]
+		},
+		{
+			"name": "Gemm.AllTransposed",
+			"attribute_maps": [
+				{ "attributes": [ { "name": "transA", "value": 1, "type": "Int32" }, { "name": "transB", "value": 1, "type": "Int32" } ] }
+			]
+		}
 	]
 	
 }
diff --git a/Engine/Plugins/Experimental/NNX/Source/NNXQA/Resources/NNXQATestDescSchema.json b/Engine/Plugins/Experimental/NNX/Source/NNXQA/Resources/NNXQATestDescSchema.json
--- a/Engine/Plugins/Experimental/NNX/Source/NNXQA/Resources/NNXQATestDescSchema.json
+++ b/Engine/Plugins/Experimental/NNX/Source/NNXQA/Resources/NNXQATestDescSchema.json
@@ -18,6 +18,10 @@
 			"type": "array",
 			"items": { "$ref": "#/$defs/test_config_input_output_set" }
 		},
+		"attribute_sets": {
+			"type": "array",
+			"items": { "$ref": "#/$defs/test_config_attribute_set" }
+		},
 		"comment": {
 			"type": "string",
 			"description": "Unused by test system, usefull to add clarification in json file"
@@ -102,6 +106,39 @@
 			}
 		},
 
+		"test_config_attribute": {
+			"type": "object",
+			"required": [ "name", "value" ],
+			"additionalProperties": false,
+			"properties": {
+				"name": {
+					"type": "string",
+					"description": "Name of the attribute"
+				},
+				"value": {
+					"type": "number",
+					"description": "Value of the attribute"
+				},
+				"type": {
+					"type": "string",
+					"description": "Type of the attribute"
+				}
+			}
+		},
+
+		"test_config_attribute_map": {
+			"type": "object",
+			"required": [ "attributes" ],
+			"additionalProperties": false,
+			"properties": {
+				"attributes": {
+					"type": "array",
+					"description": "Attributes to test",
+					"items": { "$ref": "#/$defs/test_config_attribute" }
+				}
+			}
+		},
+
 		"test_config_target": {
 			"type": "object",
 			"required": [ "target" ],
@@ -229,6 +266,32 @@
 					"description": "Unused by test system, usefull to add clarification in json file"
 				}
 			}
+		},
+
+		"test_config_attribute_set": {
+			"type": "object",
+			"required": [ "name", "attribute_maps" ],
+			"additionalProperties": false,
+			"properties": {
+				"name": {
+					"type": "string",
+					"description": "Name of this attribute set"
+				},
+				"attribute_maps": {
+					"type": "array",
+					"description": "Array of attribute maps to test",
+					"items": { "$ref": "#/$defs/test_config_attribute_map" }
+				},
+				"multiply_with_attribute_sets": {
+					"type": "array",
+					"description": "Combine with other attribute sets",
+					"items": { "type": "string" }
+				},
+				"comment": {
+					"type": "string",
+					"description": "Unused by test system, usefull to add clarification in json file"
+				}
+			}
 		}
 	}
 }
\ No newline at end of file
diff --git a/Engine/Plugins/Experimental/NNX/Source/NNXRuntimeRDG/Private/NNXRuntimeHLSLGemmOp.cpp b/Engine/Plugins/Experimental/NNX/Source/NNXRuntimeRDG/Private/NNXRuntimeHLSLGemmOp.cpp
--- a/Engine/Plugins/Experimental/NNX/Source/NNXRuntimeRDG/Private/NNXRuntimeHLSLGemmOp.cpp
+++ b/Engine/Plugins/Experimental/NNX/Source/NNXRuntimeRDG/Private/NNXRuntimeHLSLGemmOp.cpp
@@ -67,6 +67,11 @@ namespace NNX
 			// CScalar = C != nullptr ? C[0] : (InElementType)0;
 			bNoBias = InputTensors.Num() != 3 /*|| InputC.Sizes[0] * InputC.Sizes[1] < 1*/;
 
+			InputAlpha = Attributes.GetOptionalFloat(TEXT("alpha"), InputAlpha);
+			InputBeta = Attributes.GetOptionalFloat(TEXT("beta"), InputBeta);
+			InputTransA = Attributes.GetOptionalInt32(TEXT("transA"), InputTransA);
+			InputTransB = Attributes.GetOptionalInt32(TEXT("transB"), InputTransB);
+
 			return true;
 		}
 
