NNX model validation POC

Problems:
* We need to validate that a given engine/runtime support a given model at cooking time.
* We need to have dedicated optimization for a particular engine/runtime.

Solutions:
* Redefine the runtime as an engine (rename to follow in different CL)
* Embed an optimizer to the engine, optimizer contain a model validator
* Define validation as a set attribute and input type check
* Have the operators register a callback to a validation registry to be executed offline

Tests:
* all green test in editor.
* added attribute and input type validation for Gemm as a proof of concept.

diff --git a/Engine/Plugins/Experimental/NNX/Source/NNXCore/Private/NNXInferenceModel.cpp b/Engine/Plugins/Experimental/NNX/Source/NNXCore/Private/NNXInferenceModel.cpp
--- a/Engine/Plugins/Experimental/NNX/Source/NNXCore/Private/NNXInferenceModel.cpp
+++ b/Engine/Plugins/Experimental/NNX/Source/NNXCore/Private/NNXInferenceModel.cpp
@@ -2,7 +2,7 @@
 
 #include "NNXInferenceModel.h"
 
-UMLInferenceModel* UMLInferenceModel::CreateFromData(EMLInferenceFormat Format, TArrayView<uint8> Data)
+UMLInferenceModel* UMLInferenceModel::CreateFromFormatDesc(const FNNXFormatDesc& FormatDesc)
 {
 	UMLInferenceModel* Model = NewObject<UMLInferenceModel>((UObject*)GetTransientPackage(), UMLInferenceModel::StaticClass());
 
@@ -12,31 +12,11 @@ UMLInferenceModel* UMLInferenceModel::CreateFromData(EMLInferenceFormat Format,
 		return nullptr;
 	}
 
-	Model->Format = Format;
-	Model->Data.SetNum(Data.Num());
-	FMemory::Memcpy(Model->Data.GetData(), Data.GetData(), Data.Num());
-
+	Model->FormatDesc = FormatDesc;
 	return Model;
 }
 
-UMLInferenceModel::UMLInferenceModel()
-	: Format(EMLInferenceFormat::Invalid)
-{
-}
-
-EMLInferenceFormat UMLInferenceModel::GetFormat() const
+const FNNXFormatDesc& UMLInferenceModel::GetFormatDesc() const
 {
-	return Format;
-}
-
-// Return model data
-const TArray<uint8>& UMLInferenceModel::GetData() const
-{
-	return Data;
-}
-
-// Return data size in bytes
-uint32 UMLInferenceModel::GetDataSize() const
-{
-	return Data.Num();
-}
+	return FormatDesc;
+}
\ No newline at end of file
diff --git a/Engine/Plugins/Experimental/NNX/Source/NNXCore/Private/NNXModelOptimizerInterface.h b/Engine/Plugins/Experimental/NNX/Source/NNXCore/Private/NNXModelOptimizerInterface.h
new file mode 100644
--- /dev/null
+++ b/Engine/Plugins/Experimental/NNX/Source/NNXCore/Private/NNXModelOptimizerInterface.h
@@ -0,0 +1,51 @@
+// Copyright Epic Games, Inc. All Rights Reserved.
+#pragma once
+
+#include "NNXRuntimeFormat.h"
+
+namespace NNX
+{
+	
+//TODO create a dedicated FOptimizerOptionsMap when function diverge with FMLAttributeMap
+//example introduction of sparse tensor to FMLAttribute witch make no sense as an optimizer attribute
+using FOptimizerOptionsMap = FMLAttributeMap;
+
+/** Interface class for NNX model validator */
+class IModelValidator
+{
+public:
+	virtual ~IModelValidator() = default;
+	virtual FString GetName() const = 0;
+	virtual bool ValidateModel(const FNNXFormatDesc& InputModel) const = 0;
+};
+
+/** Interface class for NNX model optimizer pass */
+class IModelOptimizerPass
+{
+public:
+	virtual ~IModelOptimizerPass() = default;
+	virtual FString GetName() const = 0;
+
+	//Optimize the model in place, potentially changing the format
+	virtual bool ApplyPass(FNNXFormatDesc& Model, const FOptimizerOptionsMap& Options) const = 0;
+};
+
+/** Interface class for NNX model optimizer */
+class IModelOptimizer
+{
+public:
+
+	virtual ~IModelOptimizer() = default;
+	virtual FString GetName() const = 0;
+
+	//Allow to extend/customize an optimizer by adding passes. They should be executed in order.
+	virtual void AddOptimizationPass(TSharedPtr<IModelOptimizerPass> ModelOptimizerPass) = 0;
+	
+	//Allow to extend/customize an optimizer all validators should be run between each pass.
+	virtual void AddValidator(TSharedPtr<IModelValidator>) = 0;
+	
+	//Apply all passes and validators to the input model, produce an optimized model potentially in a different format
+	virtual bool Optimize(const FNNXFormatDesc& InputModel, FNNXFormatDesc& OutModel, const FOptimizerOptionsMap& Options) = 0;
+};
+
+} // NNX
diff --git a/Engine/Plugins/Experimental/NNX/Source/NNXCore/Public/NNXInferenceModel.h b/Engine/Plugins/Experimental/NNX/Source/NNXCore/Public/NNXInferenceModel.h
--- a/Engine/Plugins/Experimental/NNX/Source/NNXCore/Public/NNXInferenceModel.h
+++ b/Engine/Plugins/Experimental/NNX/Source/NNXCore/Public/NNXInferenceModel.h
@@ -3,18 +3,9 @@
 #pragma once
 
 #include "CoreMinimal.h"
+#include "NNXRuntimeFormat.h"
 #include "NNXInferenceModel.generated.h"
 
-/// Runtime inference format
-UENUM()
-enum class EMLInferenceFormat : uint8
-{
-	Invalid,
-	ONNX,				//!< ONNX Open Neural Network Exchange
-	ORT,				//!< ONNX Runtime (only for CPU)
-	NNXRT				//!< NNX Runtime format
-};
-
 UCLASS(BlueprintType)
 class NNXCORE_API UMLInferenceModel : public UObject
 {
@@ -22,24 +13,13 @@ class NNXCORE_API UMLInferenceModel : public UObject
 
 public:
 
-	static UMLInferenceModel* CreateFromData(EMLInferenceFormat Format, TArrayView<uint8> Data);
-
-	UMLInferenceModel();
-
-	EMLInferenceFormat GetFormat() const;
-
-	// Return model data
-	const TArray<uint8>& GetData() const;
-
-	// Return data size in bytes
-	uint32 GetDataSize() const;
+	static UMLInferenceModel* CreateFromFormatDesc(const FNNXFormatDesc& Model);
+	
+	UMLInferenceModel() = default;
+	
+	const FNNXFormatDesc& GetFormatDesc() const;
 
 private:
-
-	// TODO: Formats: ONNX, ORT, NNXRT
-	UPROPERTY(EditAnywhere, Category = "Neural Network Inference")
-	TArray<uint8>		Data;
-
 	UPROPERTY(EditAnywhere, Category = "Neural Network Inference")
-	EMLInferenceFormat	Format;
+	FNNXFormatDesc FormatDesc;
 };
diff --git a/Engine/Plugins/Experimental/NNX/Source/NNXCore/Public/NNXRuntime.h b/Engine/Plugins/Experimental/NNX/Source/NNXCore/Public/NNXRuntime.h
--- a/Engine/Plugins/Experimental/NNX/Source/NNXCore/Public/NNXRuntime.h
+++ b/Engine/Plugins/Experimental/NNX/Source/NNXCore/Public/NNXRuntime.h
@@ -23,6 +23,7 @@ class FRDGBuilder;
 namespace NNX
 {
 class FMLInferenceModel;
+class IModelOptimizer;
 
 //class IOperatorRegistry;
 
@@ -41,6 +42,7 @@ class IRuntime
 public:
 	virtual ~IRuntime() = default;
 	virtual FString GetRuntimeName() const = 0;
+	virtual TUniquePtr<IModelOptimizer> CreateModelOptimizer() const = 0;
 
 	// TODO: Replace with QuerySupport()
 	// Returns flags from ERuntimeSupportFlags
diff --git a/Engine/Plugins/Experimental/NNX/Source/NNXCore/Public/NNXRuntimeFormat.h b/Engine/Plugins/Experimental/NNX/Source/NNXCore/Public/NNXRuntimeFormat.h
--- a/Engine/Plugins/Experimental/NNX/Source/NNXCore/Public/NNXRuntimeFormat.h
+++ b/Engine/Plugins/Experimental/NNX/Source/NNXCore/Public/NNXRuntimeFormat.h
@@ -16,6 +16,30 @@ enum class EMLFormatTensorType : uint8
 	Intermediate
 };
 
+//TODO jira 167589: convert EMLInferenceFormat to a CCCC for easier extension of the framework
+UENUM()
+enum class ENNXInferenceFormat : uint8
+{
+	Invalid,
+	ONNX,				//!< ONNX Open Neural Network Exchange
+	ORT,				//!< ONNX Runtime (only for CPU)
+	NNXRT				//!< NNX Runtime format
+};
+
+USTRUCT()
+struct FNNXFormatDesc
+{
+	GENERATED_USTRUCT_BODY()
+	
+	UPROPERTY(VisibleAnywhere, Category = "Neural Network Inference")
+	TArray<uint8>		Data;
+	
+	UPROPERTY(VisibleAnywhere, Category = "Neural Network Inference")
+	ENNXInferenceFormat	Format { ENNXInferenceFormat::Invalid };
+
+	FNNXFormatDesc() = default;
+};
+
 // Required by LoadModel() when loading operators in HLSL and DirectML runtime
 USTRUCT()
 struct FMLFormatAttributeDesc
diff --git a/Engine/Plugins/Experimental/NNX/Source/NNXQA/Private/NNXQAParametricTest.cpp b/Engine/Plugins/Experimental/NNX/Source/NNXQA/Private/NNXQAParametricTest.cpp
--- a/Engine/Plugins/Experimental/NNX/Source/NNXQA/Private/NNXQAParametricTest.cpp
+++ b/Engine/Plugins/Experimental/NNX/Source/NNXQA/Private/NNXQAParametricTest.cpp
@@ -5,6 +5,7 @@
 #include "UObject/Class.h"
 #include "NNXCore.h"
 #include "NNXTypes.h"
+#include "NNXRuntimeFormat.h"
 #include "NNXQAUtils.h"
 #include "NNXQAJsonUtils.h"
 #include "NNXModelBuilder.h"
@@ -353,14 +354,14 @@ namespace Test
 
 								if (!bAtLeastAnAttributeTestWasWadded)
 								{
-									FTests::FTestSetup& Test = AddTest(TestCategoryPath, TestBaseName, TEXT(".") + GetTestSuffix(Dataset));
-									
-									ApplyRuntimesConfig(Test, TestCategory.Runtimes);
-									ApplyRuntimesConfig(Test, InputOutputSet.Runtimes);
-									ApplyTargetConfig(Test, TestTarget);
-									ApplyDatasetConfig(Test, Dataset, InputTypeFromTarget, OutputTypeFromTarget);
-									Test.IsModelTest = bIsModelCategory;
-								}
+								FTests::FTestSetup& Test = AddTest(TestCategoryPath, TestBaseName, TEXT(".") + GetTestSuffix(Dataset));
+								
+								ApplyRuntimesConfig(Test, TestCategory.Runtimes);
+								ApplyRuntimesConfig(Test, InputOutputSet.Runtimes);
+								ApplyTargetConfig(Test, TestTarget);
+								ApplyDatasetConfig(Test, Dataset, InputTypeFromTarget, OutputTypeFromTarget);
+								Test.IsModelTest = bIsModelCategory;
+							}
 							}
 							bAtLeastATestWasAdded = true;
 						}
@@ -389,31 +390,31 @@ namespace Test
 	
 	static bool RunParametricTest(FTests::FTestSetup& TestSetup, const FString& RuntimeFilter)
 	{
-		TArray<uint8> ModelData;
+		FNNXFormatDesc ONNXModel;
 
 		if (TestSetup.IsModelTest)
 		{
 			// Model test, load model from disk
 			FString ModelPath = GetFullModelPath(TestSetup.TargetName + TEXT(".onnx"));
-			const bool bIsModelInMem = FFileHelper::LoadFileToArray(ModelData, *ModelPath);
-			
+			const bool bIsModelInMem = FFileHelper::LoadFileToArray(ONNXModel.Data, *ModelPath);
 			if (!bIsModelInMem)
 			{
 				UE_LOG(LogNNX, Error, TEXT("Can't load model from disk at path '%s'. Tests ABORTED!"), *ModelPath);
 				return false;
 			}
+			ONNXModel.Format = ENNXInferenceFormat::ONNX;
 		}
 		else
 		{
 			// Operator test, create model in memory
-			if (!CreateONNXModelForOperator(TestSetup.TargetName, TestSetup.Inputs, TestSetup.Outputs, TestSetup.AttributeMap, ModelData))
+			if (!CreateONNXModelForOperator(TestSetup.TargetName, TestSetup.Inputs, TestSetup.Outputs, TestSetup.AttributeMap, ONNXModel))
 			{
 				UE_LOG(LogNNX, Error, TEXT("Failed to create model for test '%s'. Test ABORTED!"), *TestSetup.TargetName);
 				return false;
 			}
 		}
 
-		return CompareONNXModelInferenceAcrossRuntimes(ModelData, TestSetup, RuntimeFilter);
+		return CompareONNXModelInferenceAcrossRuntimes(ONNXModel, TestSetup, RuntimeFilter);
 	}
 
 	static FString AutomationRuntimeFilter;
diff --git a/Engine/Plugins/Experimental/NNX/Source/NNXQA/Private/NNXQAUtils.cpp b/Engine/Plugins/Experimental/NNX/Source/NNXQA/Private/NNXQAUtils.cpp
--- a/Engine/Plugins/Experimental/NNX/Source/NNXQA/Private/NNXQAUtils.cpp
+++ b/Engine/Plugins/Experimental/NNX/Source/NNXQA/Private/NNXQAUtils.cpp
@@ -325,43 +325,31 @@ namespace Test
 		return FMath::Sin((float)(ElementIndex + IndexOffsetBetweenTensor * TensorIndex));
 	}
 
-	static int RunTestInference(TArrayView<uint8> ONNXModelData, NNX::IRuntime* Runtime, 
+	static int RunTestInference(const FNNXFormatDesc& ONNXModelData, NNX::IRuntime* Runtime, 
 		TArray<NNX::FMLTensorDesc>& OutOutputTensorsDesc, TArray<TArray<char>>& OutOutputMemBuffers)
 	{
 		OutOutputMemBuffers.Empty();
 		OutOutputTensorsDesc.Empty();
 
-		UMLInferenceModel* UInferenceModel = nullptr;
-		if (Runtime->GetSupportFlags() == EMLRuntimeSupportFlags::RDG)
-		{
-			//Convert model from ONNX to RDG format as this Runtime only support RDG format
-			TArray<uint8> RDGModelData;
-			TUniquePtr<IMLModelOptimizer> Optimizer(CreateONNXToNNXModelOptimizer());
-			if (!Optimizer || !Optimizer->Optimize(ONNXModelData, RDGModelData))
-			{
-				UE_LOG(LogNNX, Error, TEXT("Failed to optimize the model"));
-				return false;
-			}
-			UInferenceModel = UMLInferenceModel::CreateFromData(EMLInferenceFormat::NNXRT, RDGModelData);
-		}
-		else
+		FOptimizerOptionsMap Options;
+		FNNXFormatDesc RuntimeModelData;
+		TUniquePtr<IModelOptimizer> Optimizer = Runtime->CreateModelOptimizer();
+		
+		if (!Optimizer || !Optimizer->Optimize(ONNXModelData, RuntimeModelData, Options))
 		{
-			UInferenceModel = UMLInferenceModel::CreateFromData(EMLInferenceFormat::ONNX, ONNXModelData);
+			UE_LOG(LogNNX, Error, TEXT("Failed to optimize the model"));
+			return -1;
 		}
-
+		
+		UMLInferenceModel* UInferenceModel = UMLInferenceModel::CreateFromFormatDesc(RuntimeModelData);
 		TUniquePtr<NNX::FMLInferenceModel> InferenceModel(Runtime->CreateInferenceModel(UInferenceModel));
+		
 		if (!InferenceModel.IsValid())
 		{
 			UE_LOG(LogNNX, Error, TEXT("Could not create Inference model."));
 			return -1;
 		}
 
-		if (!InferenceModel)
-		{
-			UE_LOG(LogNNX, Warning, TEXT("Error:Failed to create NNX model"));
-			return -1;
-		}
-
 		//bind tensors to memory (CPU) and initialize
 		const TArrayView<const NNX::FMLTensorDesc> InputTensorsDescView = InferenceModel->GetInputTensors();
 		TArray<NNX::FMLTensorBinding> InputBindings;
@@ -380,7 +368,7 @@ namespace Test
 		const constexpr int32 NumTensorPtrForDebug = 3;
 		float* InputsAsFloat[NumTensorPtrForDebug];
 		float* OutputsAsFloat[NumTensorPtrForDebug];
-		FMemory::Memzero(InputsAsFloat, NumTensorPtrForDebug *sizeof(float*));
+		FMemory::Memzero(InputsAsFloat, NumTensorPtrForDebug * sizeof(float*));
 		FMemory::Memzero(OutputsAsFloat, NumTensorPtrForDebug * sizeof(float*));
 		for (int32 i = 0; i < NumTensorPtrForDebug && i < InputMemBuffers.Num(); ++i)
 		{
@@ -397,7 +385,7 @@ namespace Test
 		return ReturnValue;
 	}
 
-	bool CompareONNXModelInferenceAcrossRuntimes(TArrayView<uint8> ONNXModelData, const FTests::FTestSetup& TestSetup, const FString& RuntimeFilter)
+	bool CompareONNXModelInferenceAcrossRuntimes(const FNNXFormatDesc& ONNXModel, const FTests::FTestSetup& TestSetup, const FString& RuntimeFilter)
 	{
 		FString CurrentPlatform = UGameplayStatics::GetPlatformName();
 		if (TestSetup.AutomationExcludedPlatform.Contains(CurrentPlatform))
@@ -418,7 +406,7 @@ namespace Test
 		TArray<TArray<char>> RefOutputMemBuffers;
 		TArray<NNX::FMLTensorDesc> RefOutputTensorDescs;
 
-		RunTestInference(ONNXModelData, RefRuntime, RefOutputTensorDescs, RefOutputMemBuffers);
+		RunTestInference(ONNXModel, RefRuntime, RefOutputTensorDescs, RefOutputMemBuffers);
 
 		// Test against other runtime
 		TArray<NNX::IRuntime*> Runtimes;
@@ -459,7 +447,7 @@ namespace Test
 				float AbsoluteErrorEpsilon = TestSetup.GetAbsoluteErrorEpsilonForRuntime(RuntimeName);
 				float RelativeErrorPercent = TestSetup.GetRelativeErrorPercentForRuntime(RuntimeName);
 
-				RunTestInference(ONNXModelData, Runtime, OutputTensorDescs, OutputMemBuffers);
+				RunTestInference(ONNXModel, Runtime, OutputTensorDescs, OutputMemBuffers);
 				if (OutputTensorDescs.Num() == RefOutputTensorDescs.Num())
 				{
 					for (int i = 0; i < OutputTensorDescs.Num(); ++i)
diff --git a/Engine/Plugins/Experimental/NNX/Source/NNXQA/Private/NNXQAUtils.h b/Engine/Plugins/Experimental/NNX/Source/NNXQA/Private/NNXQAUtils.h
--- a/Engine/Plugins/Experimental/NNX/Source/NNXQA/Private/NNXQAUtils.h
+++ b/Engine/Plugins/Experimental/NNX/Source/NNXQA/Private/NNXQAUtils.h
@@ -2,10 +2,11 @@
 
 #pragma once
 
-#include "Containers/ArrayView.h"
 #include "Containers/UnrealString.h"
 #include "NNXRuntime.h"
 
+struct FNNXFormatDesc;
+
 namespace NNX 
 {
 namespace Test 
@@ -59,7 +60,7 @@ namespace Test
 		TArray<FTestSetup> TestSetups;
 	};
 
-	bool CompareONNXModelInferenceAcrossRuntimes(TArrayView<uint8> ModelData, const FTests::FTestSetup& TestSetup, const FString& RuntimeFilter = TEXT(""));
+	bool CompareONNXModelInferenceAcrossRuntimes(const FNNXFormatDesc& ONNXModel, const FTests::FTestSetup& TestSetup, const FString& RuntimeFilter = TEXT(""));
 
 	FString FMLTensorDescToString(const FMLTensorDesc& TensorDesc);
 	FString TensorToString(const FMLTensorDesc& TensorDesc, const TArray<char>& TensorData);
diff --git a/Engine/Plugins/Experimental/NNX/Source/NNXRuntimeCPU/NNXRuntimeCPU.Build.cs b/Engine/Plugins/Experimental/NNX/Source/NNXRuntimeCPU/NNXRuntimeCPU.Build.cs
--- a/Engine/Plugins/Experimental/NNX/Source/NNXRuntimeCPU/NNXRuntimeCPU.Build.cs
+++ b/Engine/Plugins/Experimental/NNX/Source/NNXRuntimeCPU/NNXRuntimeCPU.Build.cs
@@ -28,6 +28,7 @@ public class NNXRuntimeCPU : ModuleRules
 			(
 			new string[] {
 				//"ORTDefault",
+                "NNXUtils",
 				"NNX_ONNXRuntime"	// Select this for open-source-based ONNX Runtime
 			}
 		);
diff --git a/Engine/Plugins/Experimental/NNX/Source/NNXRuntimeCPU/Private/NNXRuntimeCPU.cpp b/Engine/Plugins/Experimental/NNX/Source/NNXRuntimeCPU/Private/NNXRuntimeCPU.cpp
--- a/Engine/Plugins/Experimental/NNX/Source/NNXRuntimeCPU/Private/NNXRuntimeCPU.cpp
+++ b/Engine/Plugins/Experimental/NNX/Source/NNXRuntimeCPU/Private/NNXRuntimeCPU.cpp
@@ -1,6 +1,7 @@
 // Copyright Epic Games, Inc. All Rights Reserved.
 #include "NNXRuntimeCPU.h"
 #include "NNXRuntimeCPUUtils.h"
+#include "NNXModelOptimizer.h"
 #include "NeuralTimer.h"
 #include "RedirectCoutAndCerrToUeLog.h"
 
@@ -11,6 +12,11 @@ FString FRuntimeCPU::GetRuntimeName() const
 	return NNX_RUNTIME_CPU_NAME;
 }
 
+TUniquePtr<IModelOptimizer> FRuntimeCPU::CreateModelOptimizer() const
+{
+	return CreateONNXToONNXModelOptimizer();
+}
+
 EMLRuntimeSupportFlags FRuntimeCPU::GetSupportFlags() const 
 {
 	return EMLRuntimeSupportFlags::CPU;
@@ -66,11 +72,12 @@ bool FMLInferenceModelCPU::Init(const UMLInferenceModel* InferenceModel)
 	
 	// Clean previous networks
 	bIsLoaded = false;
-	const TArray<uint8>& ModelBuffer{ InferenceModel->GetData() };
+	const TArray<uint8>& ModelBuffer{ InferenceModel->GetFormatDesc().Data };
 
 	// Checking Inference Model 
 	{
-		if (ModelBuffer.Num() == 0) {
+		if (ModelBuffer.Num() == 0) 
+		{
 			UE_LOG(LogNNX, Warning, TEXT("FMLInferenceModelCPU::Load(): Input model path is empty."));
 			return false;
 		}
diff --git a/Engine/Plugins/Experimental/NNX/Source/NNXRuntimeCPU/Public/NNXRuntimeCPU.h b/Engine/Plugins/Experimental/NNX/Source/NNXRuntimeCPU/Public/NNXRuntimeCPU.h
--- a/Engine/Plugins/Experimental/NNX/Source/NNXRuntimeCPU/Public/NNXRuntimeCPU.h
+++ b/Engine/Plugins/Experimental/NNX/Source/NNXRuntimeCPU/Public/NNXRuntimeCPU.h
@@ -34,6 +34,7 @@ namespace NNX
 		bool Init();
 		virtual ~FRuntimeCPU();
 		virtual FString GetRuntimeName() const;
+		virtual TUniquePtr<IModelOptimizer> CreateModelOptimizer() const;
 		virtual EMLRuntimeSupportFlags GetSupportFlags() const;
 		virtual FMLInferenceModel* CreateInferenceModel(UMLInferenceModel* Model);
 		FMLInferenceModel* CreateInferenceModel(UMLInferenceModel* InModel, const FMLInferenceNNXCPUConf& InConf);
diff --git a/Engine/Plugins/Experimental/NNX/Source/NNXRuntimeORT/NNXRuntimeORT.Build.cs b/Engine/Plugins/Experimental/NNX/Source/NNXRuntimeORT/NNXRuntimeORT.Build.cs
--- a/Engine/Plugins/Experimental/NNX/Source/NNXRuntimeORT/NNXRuntimeORT.Build.cs
+++ b/Engine/Plugins/Experimental/NNX/Source/NNXRuntimeORT/NNXRuntimeORT.Build.cs
@@ -33,6 +33,7 @@ public class NNXRuntimeORT : ModuleRules
 				"ORTHelper",
 				//"ONNXRuntime",
 				"ORTDefault",
+                "NNXUtils"
 			}
 		);
 
diff --git a/Engine/Plugins/Experimental/NNX/Source/NNXRuntimeORT/Private/NNXRuntimeORT.cpp b/Engine/Plugins/Experimental/NNX/Source/NNXRuntimeORT/Private/NNXRuntimeORT.cpp
--- a/Engine/Plugins/Experimental/NNX/Source/NNXRuntimeORT/Private/NNXRuntimeORT.cpp
+++ b/Engine/Plugins/Experimental/NNX/Source/NNXRuntimeORT/Private/NNXRuntimeORT.cpp
@@ -6,6 +6,7 @@
 
 #include "NeuralTimer.h"
 #include "NNXRuntimeORTUtils.h"
+#include "NNXModelOptimizer.h"
 #include "RedirectCoutAndCerrToUeLog.h"
 
 // NOTE: For now we only have DML on Windows, we should add support for XSX
@@ -57,6 +58,23 @@ EMLRuntimeSupportFlags FRuntimeORTDml::GetSupportFlags() const
 }
 #endif
 
+TUniquePtr<IModelOptimizer> FRuntimeORTCpu::CreateModelOptimizer() const
+{
+	return CreateONNXToONNXModelOptimizer();
+}
+
+#if PLATFORM_WINDOWS
+TUniquePtr<IModelOptimizer> FRuntimeORTCuda::CreateModelOptimizer() const
+{
+	return CreateONNXToONNXModelOptimizer();
+}
+
+TUniquePtr<IModelOptimizer> FRuntimeORTDml::CreateModelOptimizer() const
+{
+	return CreateONNXToONNXModelOptimizer();
+}
+#endif
+
 FMLInferenceModel* FRuntimeORTCpu::CreateInferenceModel(UMLInferenceModel* InModel, const FMLInferenceNNXORTConf& InConf)
 {
 	FMLInferenceModelORTCpu* ORTModel = new FMLInferenceModelORTCpu(&NNXEnvironmentORT, InConf);
@@ -132,7 +150,7 @@ bool FMLInferenceModelORT::Init(const UMLInferenceModel* InferenceModel)
 	
 	// Clean previous networks
 	bIsLoaded = false;
-	const TArray<uint8>& ModelBuffer{ InferenceModel->GetData()};
+	const TArray<uint8>& ModelBuffer{ InferenceModel->GetFormatDesc().Data };
 
 	// Checking Inference Model 
 	{
diff --git a/Engine/Plugins/Experimental/NNX/Source/NNXRuntimeORT/Private/NNXRuntimeORT.h b/Engine/Plugins/Experimental/NNX/Source/NNXRuntimeORT/Private/NNXRuntimeORT.h
--- a/Engine/Plugins/Experimental/NNX/Source/NNXRuntimeORT/Private/NNXRuntimeORT.h
+++ b/Engine/Plugins/Experimental/NNX/Source/NNXRuntimeORT/Private/NNXRuntimeORT.h
@@ -56,6 +56,7 @@ namespace NNX
 		FRuntimeORTCpu() = default;
 		virtual ~FRuntimeORTCpu() = default;
 		virtual FString GetRuntimeName() const override;
+		virtual TUniquePtr<IModelOptimizer> CreateModelOptimizer() const override;
 		virtual EMLRuntimeSupportFlags GetSupportFlags() const override;
 		virtual FMLInferenceModel* CreateInferenceModel(UMLInferenceModel* InModel);
 		FMLInferenceModel* CreateInferenceModel(UMLInferenceModel* InModel, const FMLInferenceNNXORTConf& InConf);
@@ -69,6 +70,7 @@ namespace NNX
 		FRuntimeORTCuda() = default;
 		virtual ~FRuntimeORTCuda() = default;
 		virtual FString GetRuntimeName() const override;
+		virtual TUniquePtr<IModelOptimizer> CreateModelOptimizer() const override;
 		virtual EMLRuntimeSupportFlags GetSupportFlags() const override;
 		virtual FMLInferenceModel* CreateInferenceModel(UMLInferenceModel* InModel);
 		FMLInferenceModel* CreateInferenceModel(UMLInferenceModel* InModel, const FMLInferenceNNXORTConf& InConf);
@@ -81,6 +83,7 @@ namespace NNX
 		FRuntimeORTDml() = default;
 		virtual ~FRuntimeORTDml() = default;
 		virtual FString GetRuntimeName() const override;
+		virtual TUniquePtr<IModelOptimizer> CreateModelOptimizer() const override;
 		virtual EMLRuntimeSupportFlags GetSupportFlags() const override;
 		virtual FMLInferenceModel* CreateInferenceModel(UMLInferenceModel* InModel);
 		FMLInferenceModel* CreateInferenceModel(UMLInferenceModel* InModel, const FMLInferenceNNXORTConf& InConf);
diff --git a/Engine/Plugins/Experimental/NNX/Source/NNXRuntimeRDG/NNXRuntimeRDG.Build.cs b/Engine/Plugins/Experimental/NNX/Source/NNXRuntimeRDG/NNXRuntimeRDG.Build.cs
--- a/Engine/Plugins/Experimental/NNX/Source/NNXRuntimeRDG/NNXRuntimeRDG.Build.cs
+++ b/Engine/Plugins/Experimental/NNX/Source/NNXRuntimeRDG/NNXRuntimeRDG.Build.cs
@@ -23,7 +23,8 @@ public class NNXRuntimeRDG : ModuleRules
         PrivateDependencyModuleNames.AddRange(new string[]
         {
             "NNXCore",
-            "NNXHlslShaders",
+			"NNXUtils",
+			"NNXHlslShaders",
             "RHI"
         });
 
diff --git a/Engine/Plugins/Experimental/NNX/Source/NNXRuntimeRDG/Private/NNXRuntimeDML.cpp b/Engine/Plugins/Experimental/NNX/Source/NNXRuntimeRDG/Private/NNXRuntimeDML.cpp
--- a/Engine/Plugins/Experimental/NNX/Source/NNXRuntimeRDG/Private/NNXRuntimeDML.cpp
+++ b/Engine/Plugins/Experimental/NNX/Source/NNXRuntimeRDG/Private/NNXRuntimeDML.cpp
@@ -3,6 +3,7 @@
 #include "NNXRuntimeRDG.h"
 #include "NNXInferenceModel.h"
 #include "NNXRuntimeFormat.h"
+#include "NNXModelOptimizer.h"
 #include "NNXOperator.h"
 
 #include "HAL/FileManager.h"
@@ -512,7 +513,7 @@ protected:
 /**
  * DirectML ML operator registry
  */
-using FMLOperatorRegistryDml = TMLOperatorRegistryRDG<FMLOperatorDml>;
+using FMLOperatorRegistryDml = TOperatorRegistryRDG<FMLOperatorDml>;
 
 /**
  * Element-wise unary ML operator implementation
@@ -1002,6 +1003,15 @@ public:
 		return NNX_RUNTIME_DML_NAME;
 	}
 
+	virtual TUniquePtr<IModelOptimizer> CreateModelOptimizer() const
+	{
+		return CreateONNXToNNXModelOptimizer();
+		
+		//TODO jira 167581: Add DML specific model validator
+		//see ModelOptimizer->AddValidator(MakeShared<FModelValidatorHlsl>());
+		//in NNXRuntimeHLSL.cpp
+	}
+	
 	virtual EMLRuntimeSupportFlags GetSupportFlags() const override
 	{
 		return EMLRuntimeSupportFlags::RDG;
@@ -1247,9 +1257,10 @@ FMLInferenceModelDml::~FMLInferenceModelDml()
 //
 bool FMLInferenceModelDml::Init(UMLInferenceModel* InModel, FDeviceContextDml* InDevCtx)
 {
+	check(InModel != nullptr);
 	FMLRuntimeFormat	Format;
 
-	if (!LoadModel(InModel, Format))
+	if (!LoadModel(InModel->GetFormatDesc(), Format))
 	{
 		return false;
 	}
@@ -1316,7 +1327,7 @@ FMLOperatorDml* FMLInferenceModelDml::OpCreate(const FString& OpName, TArrayView
 {
 	// TODO: Check if D3D12 device is valid
 
-	FMLOperatorRegistryDml::MLOperatorCreateFunc CreateFn = FMLOperatorRegistryDml::Get()->OpFind(OpName);
+	FMLOperatorRegistryDml::OperatorCreateFunc CreateFn = FMLOperatorRegistryDml::Get()->OpFind(OpName);
 
 	if (!CreateFn)
 	{
diff --git a/Engine/Plugins/Experimental/NNX/Source/NNXRuntimeRDG/Private/NNXRuntimeHLSL.cpp b/Engine/Plugins/Experimental/NNX/Source/NNXRuntimeRDG/Private/NNXRuntimeHLSL.cpp
--- a/Engine/Plugins/Experimental/NNX/Source/NNXRuntimeRDG/Private/NNXRuntimeHLSL.cpp
+++ b/Engine/Plugins/Experimental/NNX/Source/NNXRuntimeRDG/Private/NNXRuntimeHLSL.cpp
@@ -3,6 +3,7 @@
 #include "NNXElementWiseBinaryCS.h"
 #include "NNXInferenceModel.h"
 #include "NNXRuntimeFormat.h"
+#include "NNXModelOptimizer.h"
 #include "NNXRuntimeHLSLElementWiseBinaryOps.h"
 #include "NNXRuntimeHLSLElementWiseUnaryOps.h"
 #include "NNXRuntimeHLSLElementWiseVariadicOps.h"
@@ -61,9 +62,10 @@ FMLInferenceModelHlsl::~FMLInferenceModelHlsl()
 //
 bool FMLInferenceModelHlsl::Init(UMLInferenceModel* InModel)
 {
+	check(InModel != nullptr);
 	FMLRuntimeFormat	Format;
 
-	if (!LoadModel(InModel, Format))
+	if (!LoadModel(InModel->GetFormatDesc(), Format))
 	{
 		return false;
 	}
@@ -111,7 +113,7 @@ bool FMLInferenceModelHlsl::Init(UMLInferenceModel* InModel)
 //
 FMLOperatorHlsl* FMLInferenceModelHlsl::OpCreate(const FString& OpName, TArrayView<const FMLTensorDesc> InputTensorDescs, TArrayView<const FMLTensorDesc> OutputTensorDescs, const FMLAttributeMap& AttributeMap)
 {
-	FMLOperatorRegistryHlsl::MLOperatorCreateFunc CreateFn = FMLOperatorRegistryHlsl::Get()->OpFind(OpName);
+	FMLOperatorRegistryHlsl::OperatorCreateFunc CreateFn = FMLOperatorRegistryHlsl::Get()->OpFind(OpName);
 
 	if (!CreateFn)
 	{
@@ -173,6 +175,16 @@ public:
 		return NNX_RUNTIME_HLSL_NAME;
 	}
 
+	virtual TUniquePtr<IModelOptimizer> CreateModelOptimizer() const
+	{
+		TUniquePtr<IModelOptimizer> ModelOptimizer = CreateONNXToNNXModelOptimizer();
+		check(ModelOptimizer.IsValid());
+		
+		ModelOptimizer->AddValidator(MakeShared<FModelValidatorHlsl>());
+		
+		return ModelOptimizer;
+	}
+
 	// Returns flags from ERuntimeSupportFlags
 	virtual EMLRuntimeSupportFlags GetSupportFlags() const override
 	{
diff --git a/Engine/Plugins/Experimental/NNX/Source/NNXRuntimeRDG/Private/NNXRuntimeHLSLGemmOp.cpp b/Engine/Plugins/Experimental/NNX/Source/NNXRuntimeRDG/Private/NNXRuntimeHLSLGemmOp.cpp
--- a/Engine/Plugins/Experimental/NNX/Source/NNXRuntimeRDG/Private/NNXRuntimeHLSLGemmOp.cpp
+++ b/Engine/Plugins/Experimental/NNX/Source/NNXRuntimeRDG/Private/NNXRuntimeHLSLGemmOp.cpp
@@ -20,6 +20,27 @@ namespace NNX
 			return new FMLOperatorHlslGemm();
 		}
 
+		static bool Validate(const FMLAttributeMap& AttributeMap, TConstArrayView<EMLTensorDataType> InputTypes)
+		{
+			bool bIsValid = true;
+
+			FAttributeValidator AttributeValidator;
+			AttributeValidator.AddOptional(TEXT("alpha"), EMLAttributeDataType::Float);
+			AttributeValidator.AddOptional(TEXT("beta"), EMLAttributeDataType::Float);
+			AttributeValidator.AddOptional(TEXT("transA"), EMLAttributeDataType::Int32);
+			AttributeValidator.AddOptional(TEXT("transB"), EMLAttributeDataType::Int32);
+			bIsValid &= AttributeValidator.Validate(AttributeMap);
+			
+			FInputValidator InputValidator;
+			InputValidator.AddSupportedType(EMLTensorDataType::Float);
+			InputValidator.AddRequired();
+			InputValidator.AddRequired();
+			InputValidator.AddOptional();
+			bIsValid &= InputValidator.Validate(InputTypes);
+
+			return bIsValid;
+		}
+
 		virtual ~FMLOperatorHlslGemm() = default;
 
 	private:
@@ -114,7 +135,7 @@ namespace NNX
 
 	bool RegisterGemmOperator(FMLOperatorRegistryHlsl& Registry)
 	{
-		Registry.OpAdd(TEXT("Gemm"), FMLOperatorHlslGemm::Create);
+		Registry.OpAdd(TEXT("Gemm"), FMLOperatorHlslGemm::Create, FMLOperatorHlslGemm::Validate);
 
 		return true;
 	}
diff --git a/Engine/Plugins/Experimental/NNX/Source/NNXRuntimeRDG/Private/NNXRuntimeHLSLOp.h b/Engine/Plugins/Experimental/NNX/Source/NNXRuntimeRDG/Private/NNXRuntimeHLSLOp.h
--- a/Engine/Plugins/Experimental/NNX/Source/NNXRuntimeRDG/Private/NNXRuntimeHLSLOp.h
+++ b/Engine/Plugins/Experimental/NNX/Source/NNXRuntimeRDG/Private/NNXRuntimeHLSLOp.h
@@ -20,6 +20,7 @@ namespace NNX
 	/**
 	  * HLSL ML operator registry
 	*/
-	typedef TMLOperatorRegistryRDG<FMLOperatorHlsl> FMLOperatorRegistryHlsl;
+	typedef TOperatorRegistryRDG<FMLOperatorHlsl> FMLOperatorRegistryHlsl;
+	typedef TModelValidatorRDG<FMLOperatorHlsl> FModelValidatorHlsl;
 
 } // NNX
\ No newline at end of file
diff --git a/Engine/Plugins/Experimental/NNX/Source/NNXRuntimeRDG/Private/NNXRuntimeRDG.cpp b/Engine/Plugins/Experimental/NNX/Source/NNXRuntimeRDG/Private/NNXRuntimeRDG.cpp
--- a/Engine/Plugins/Experimental/NNX/Source/NNXRuntimeRDG/Private/NNXRuntimeRDG.cpp
+++ b/Engine/Plugins/Experimental/NNX/Source/NNXRuntimeRDG/Private/NNXRuntimeRDG.cpp
@@ -12,6 +12,152 @@
 namespace NNX
 {
 
+//
+//
+//
+bool AlwaysValidValidationFunction(const FMLAttributeMap& AttributeMap, TConstArrayView<EMLTensorDataType> InputTensorTypes)
+{
+	return true;
+}
+
+
+//
+//
+//
+FInputValidator::FInputValidator() : 
+	NumRequiredInput(0), NumOptionalInput(0)
+{
+	TemplateTypes.SetNum(1);
+}
+
+bool FInputValidator::Validate(TConstArrayView<EMLTensorDataType> InputTypes)
+{
+	check(InputTemplateIndices.Num() == NumRequiredInput + NumOptionalInput);
+
+	bool bAreInputValid = true;
+	int32 NumInputsToValidate = FMath::Min(InputTemplateIndices.Num(), InputTypes.Num());
+	
+	if (InputTypes.Num() < NumRequiredInput)
+	{
+		UE_LOG(LogNNX, Warning, TEXT("Required '%d' inputs but found '%d'."), NumRequiredInput, InputTypes.Num());
+		bAreInputValid = false;
+	}
+	if (InputTypes.Num() > NumRequiredInput+NumOptionalInput)
+	{
+		UE_LOG(LogNNX, Warning, TEXT("Got a total of '%d' inputs but should have '%d' maximum."), InputTypes.Num(), NumRequiredInput + NumOptionalInput);
+		bAreInputValid = false;
+	}
+	
+	for (int32 Idx = 0; Idx < NumInputsToValidate; ++Idx)
+	{
+		const int32 TemplateIdx = InputTemplateIndices[Idx];
+		
+		check(TemplateIdx < TemplateTypes.Num());
+		if (INDEX_NONE == TemplateTypes[TemplateIdx].Find(InputTypes[Idx]))
+		{
+			UE_LOG(LogNNX, Warning, TEXT("Input '%d' from template idx '%d' is of type '%d' is not supported."), Idx, TemplateIdx, (int)InputTypes[Idx]);
+			bAreInputValid = false;
+		}
+	}
+	return bAreInputValid;
+}
+void FInputValidator::SetTemplateCount(int TemplateCount)
+{
+	TemplateTypes.SetNum(TemplateCount);
+}
+void FInputValidator::AddSupportedType(EMLTensorDataType Type, int TemplateIdx)
+{
+	check(TemplateTypes.Num() > TemplateIdx);
+	TemplateTypes[TemplateIdx].Add(Type);
+}
+void FInputValidator::AddOptional(int32 TemplateIdx)
+{
+	InputTemplateIndices.Add(TemplateIdx);
+	++NumOptionalInput;
+}
+
+void FInputValidator::AddRequired(int32 TemplateIdx)
+{
+	checkf(NumOptionalInput==0, TEXT("All required attribute should be declared before the optional ones as they are referenced by indices"));
+	InputTemplateIndices.Add(TemplateIdx);
+	++NumRequiredInput;
+}
+
+
+//
+//
+//
+void FAttributeValidator::AddOptional(const FString& Name, EMLAttributeDataType Type)
+{
+	checkf(nullptr == OptionalAttributes.FindByPredicate([Name](const FEntry& Other) { return Other.Name == Name; }), TEXT("Attribute name should be unique"));
+	checkf(nullptr == RequiredAttributes.FindByPredicate([Name](const FEntry& Other) { return Other.Name == Name; }), TEXT("Attribute name should be unique"));
+	OptionalAttributes.Emplace(Name, Type);
+}
+
+void FAttributeValidator::AddRequired(const FString& Name, EMLAttributeDataType Type)
+{
+	checkf(nullptr == OptionalAttributes.FindByPredicate([Name](const FEntry& Other) { return Other.Name == Name; }), TEXT("Attribute name should be unique"));
+	checkf(nullptr == RequiredAttributes.FindByPredicate([Name](const FEntry& Other) { return Other.Name == Name; }), TEXT("Attribute name should be unique"));
+	RequiredAttributes.Emplace(Name, Type);
+}
+
+bool FAttributeValidator::Validate(const FMLAttributeMap& AttributesToValidate)
+{
+	bool bAreAttributesValid = true;
+
+	//Verify all required attribute are matching specifications
+	for (int32 Idx = 0; Idx < RequiredAttributes.Num(); ++Idx)
+	{
+		const FMLAttributeValue* FoundAttribute = AttributesToValidate.GetAttributeValue(RequiredAttributes[Idx].Name);
+		
+		if (FoundAttribute == nullptr)
+		{
+			bAreAttributesValid = false;
+			UE_LOG(LogNNX, Warning, TEXT("Required attribute '%s' not found."),
+				*RequiredAttributes[Idx].Name);
+		}
+		else if (RequiredAttributes[Idx].Type != FoundAttribute->GetType())
+		{
+			bAreAttributesValid = false;
+			UE_LOG(LogNNX, Warning, TEXT("Required attribute '%s' type '%d' does not match expected type '%d'."),
+				*RequiredAttributes[Idx].Name,
+				(int)FoundAttribute->GetType(),
+				(int)RequiredAttributes[Idx].Type);
+		}
+	}
+
+	//Verify all optional attribute are matching specifications
+	for (int32 Idx = 0; Idx < OptionalAttributes.Num(); ++Idx)
+	{
+		const FMLAttributeValue* FoundAttribute = AttributesToValidate.GetAttributeValue(OptionalAttributes[Idx].Name);
+		
+		if ((FoundAttribute != nullptr) && (OptionalAttributes[Idx].Type != FoundAttribute->GetType()))
+		{
+			bAreAttributesValid = false;
+			UE_LOG(LogNNX, Warning, TEXT("Optional attribute '%s' type '%d' does not match expected type '%d'."),
+				*OptionalAttributes[Idx].Name,
+				(int)FoundAttribute->GetType(),
+				(int)OptionalAttributes[Idx].Type);
+		}
+	}
+
+	//Verify all attributes are either required or optional, otherwise they are unsupported
+	for (int32 Idx = 0; Idx < AttributesToValidate.Num(); ++Idx)
+	{
+		const FString& Name = AttributesToValidate.GetName(Idx);
+		const FEntry* OptionalAttribute = OptionalAttributes.FindByPredicate([Name](const FEntry& Other) { return Other.Name == Name; });
+		const FEntry* RequiredAttribute = RequiredAttributes.FindByPredicate([Name](const FEntry& Other) { return Other.Name == Name; });
+		
+		if (OptionalAttribute == nullptr && RequiredAttribute == nullptr)
+		{
+			bAreAttributesValid = false;
+			UE_LOG(LogNNX, Warning, TEXT("Found unsupported attribute '%s'."), *Name);
+		}
+	}
+
+	return bAreAttributesValid;
+}
+	
 //
 //
 //
@@ -30,17 +176,17 @@ FMLInferenceModelRDG::~FMLInferenceModelRDG()
 //
 //
 //
-bool FMLInferenceModelRDG::LoadModel(UMLInferenceModel* InModel, FMLRuntimeFormat& Format)
+bool FMLInferenceModelRDG::LoadModel(const FNNXFormatDesc& InModel, FMLRuntimeFormat& Format)
 {
-	EMLInferenceFormat FormatType = InModel->GetFormat();
+	ENNXInferenceFormat FormatType = InModel.Format;
 
-	if (FormatType != EMLInferenceFormat::NNXRT)
+	if (FormatType != ENNXInferenceFormat::NNXRT)
 	{
 		UE_LOG(LogNNX, Warning, TEXT("Unsupported format type for NNX inference model"));
 		return false;
 	}
 
-	FMemoryReader Reader(InModel->GetData());
+	FMemoryReader Reader(InModel.Data);
 
 	FMLRuntimeFormat::StaticStruct()->SerializeBin(Reader, &Format);
 
diff --git a/Engine/Plugins/Experimental/NNX/Source/NNXRuntimeRDG/Private/NNXRuntimeRDG.h b/Engine/Plugins/Experimental/NNX/Source/NNXRuntimeRDG/Private/NNXRuntimeRDG.h
--- a/Engine/Plugins/Experimental/NNX/Source/NNXRuntimeRDG/Private/NNXRuntimeRDG.h
+++ b/Engine/Plugins/Experimental/NNX/Source/NNXRuntimeRDG/Private/NNXRuntimeRDG.h
@@ -4,8 +4,10 @@
 
 #include "NNXCore.h"
 #include "NNXRuntime.h"
+#include "NNXModelOptimizer.h"
 #include "ShaderParameterUtils.h"
 #include "RHIGPUReadback.h"
+#include "Serialization/MemoryReader.h"
 
 #include "Containers/Map.h"
 
@@ -24,6 +26,7 @@ END_SHADER_PARAMETER_STRUCT()
 
 class FRDGBuilder;
 struct FMLRuntimeFormat;
+struct FNNXFormatDesc;
 
 namespace NNX
 {
@@ -64,7 +67,7 @@ protected:
 
 	FMLInferenceModelRDG();
 
-	bool LoadModel(UMLInferenceModel* InModel, FMLRuntimeFormat& Format);
+	bool LoadModel(const FNNXFormatDesc& InModel, FMLRuntimeFormat& Format);
 
 	int SetTensors(FRDGBuilder& GraphBuilder, FMLTensorBindingArray& OutRDGBindings, FMLIntArray& OutIndices, TArrayView<const FMLTensorBinding> InBindings, TArrayView<const FMLTensorDesc> InTensors);
 
@@ -77,27 +80,75 @@ protected:
 	bool			bUseManualTransitions;
 };
 
+//TODO jira 167584 remove default validation and declare contract in all HLSL operator (see HLSL Gemm for current example)
+//TODO jira remove default validation and declare contract in all DML operator (see HLSL Gemm for current example)
+bool AlwaysValidValidationFunction(const FMLAttributeMap& AttributeMap, TConstArrayView<EMLTensorDataType> InputTypes);
+
+class FInputValidator
+{
+public:
+	FInputValidator();
+	void AddRequired(int32 TemplateIdx = 0);
+	void AddOptional(int32 TemplateIdx = 0);
+	void SetTemplateCount(int32 TemplateCount);
+	void AddSupportedType(EMLTensorDataType Type, int32 TemplateIdx = 0);
+	bool Validate(TConstArrayView<EMLTensorDataType> InputTypes);
+
+private:
+	TArray<TArray<EMLTensorDataType>> TemplateTypes;
+	TArray<int32> InputTemplateIndices;
+	int32 NumRequiredInput;
+	int32 NumOptionalInput;
+};
+
+class FAttributeValidator
+{
+public:
+	void AddOptional(const FString& Name, EMLAttributeDataType Type);
+	void AddRequired(const FString& Name, EMLAttributeDataType Type);
+	bool Validate(const FMLAttributeMap& AttributesToValidate);
+
+private:
+	struct FEntry
+	{
+		FEntry(const FString& InName, EMLAttributeDataType InType)
+			: Name(InName), Type(InType)
+		{
+		}
+
+		//TODO should be extended as needed by operator to support more validation especially around the range of the values.
+		//An example is ConvTranspose `auto_pad` enum style string that can only take a few values
+		//In the same direction we might only support a range of value for a float (for example
+		//we only support integer but the type is float, or only positive values for an int32)
+		FString Name;
+		EMLAttributeDataType Type;
+	};
+
+	TArray<FEntry> RequiredAttributes;
+	TArray<FEntry> OptionalAttributes;
+};
 
 /**
  * Registry for RDG ML operators
  */
-template<class TMLOperatorType>
-class TMLOperatorRegistryRDG
+template<class TOperatorType>
+class TOperatorRegistryRDG
 {
 public:
 
-	typedef TMLOperatorType* (*MLOperatorCreateFunc)();
+	typedef TOperatorType* (*OperatorCreateFunc)();
+	typedef bool (*OperatorValidateFunc)(const FMLAttributeMap& AttributeMap, TConstArrayView<EMLTensorDataType> InputTypes);
 
-	static TMLOperatorRegistryRDG* Get()
+	static TOperatorRegistryRDG* Get()
 	{
-		static TMLOperatorRegistryRDG Inst;
+		static TOperatorRegistryRDG Inst;
 
 		return &Inst;
 	}
 
-	MLOperatorCreateFunc OpFind(const FString& Name)
+	OperatorValidateFunc OpFindValidation(const FString& Name)
 	{
-		MLOperatorCreateFunc* Fn = Operators.Find(Name);
+		OperatorValidateFunc* Fn = OperatorValidations.Find(Name);
 
 		if (!Fn)
 		{
@@ -108,7 +159,20 @@ public:
 		return *Fn;
 	}
 
-	bool OpAdd(const FString& Name, MLOperatorCreateFunc Func)
+	OperatorCreateFunc OpFind(const FString& Name)
+	{
+		OperatorCreateFunc* Fn = Operators.Find(Name);
+
+		if (!Fn)
+		{
+			UE_LOG(LogNNX, Warning, TEXT("RDG MLOperator:%s is not registered"), *Name);
+			return nullptr;
+		}
+
+		return *Fn;
+	}
+
+	bool OpAdd(const FString& Name, OperatorCreateFunc Func, OperatorValidateFunc ValidateFunc = AlwaysValidValidationFunction)
 	{
 		if (Operators.Find(Name) != nullptr)
 		{
@@ -117,14 +181,90 @@ public:
 		}
 
 		Operators.Add(Name, Func);
+		OperatorValidations.Add(Name, ValidateFunc);
 		return true;
 	}
 
 private:
 
-	TMap<FString, MLOperatorCreateFunc> Operators;
+	TMap<FString, OperatorCreateFunc> Operators;
+	TMap<FString, OperatorValidateFunc> OperatorValidations;
 };
 
+/**
+ * Validator for RDG ML operators
+ */
+template<class TOperatorType>
+class TModelValidatorRDG : public IModelValidator
+{
+public:
+	virtual FString GetName() const 
+	{
+		return TEXT("RDG Model validator");
+	}
+	
+	virtual bool ValidateModel(const FNNXFormatDesc& InputModel) const override
+	{
+		FMLRuntimeFormat	Format;
+
+		ENNXInferenceFormat FormatType = InputModel.Format;
+		if (FormatType != ENNXInferenceFormat::NNXRT)
+		{
+			UE_LOG(LogNNX, Warning, TEXT("Unsupported format type for validator %s"), *GetName());
+			return false;
+		}
+
+		FMemoryReader Reader(InputModel.Data);
+		FMLRuntimeFormat::StaticStruct()->SerializeBin(Reader, &Format);
+
+		TOperatorRegistryRDG<TOperatorType>* Registry = TOperatorRegistryRDG<TOperatorType>::Get();
+		check(Registry != nullptr);
+
+		if (Format.Operators.Num() > 1)
+		{
+			UE_LOG(LogNNX, Warning, TEXT("Failed to validate model, currently only single layer models are supported"));
+			return false;
+		}
+		
+		// HACK: This works only for single layer networks
+		// we should get the output tensor type from intermediate tensor 
+		// (can we get it from model definition?, otherwise will need
+		// to infer from input shape type + forward propagation and op callbacks)
+		TArray<EMLTensorDataType> InputTensorTypes;
+		for (int32 Idx = 0; Idx < Format.Tensors.Num(); ++Idx)
+		{
+			if (Format.Tensors[Idx].Type == EMLFormatTensorType::Input)
+			{
+				InputTensorTypes.Add(Format.Tensors[Idx].DataType);
+			}
+		}
+
+		for (int32 Idx = 0; Idx < Format.Operators.Num(); ++Idx)
+		{
+			FMLAttributeMap AttributeMap;
+			for (const FMLFormatAttributeDesc& Desc : Format.Operators[Idx].Attributes)
+			{
+				AttributeMap.SetAttribute(Desc.Name, Desc.Value);
+			}
+
+			const FString& OpType = Format.Operators[Idx].TypeName;
+			
+			//TODO jira 167587: we should extract constant tensor from the model
+			//and pass them to the operator validation so that it can validate
+			//the shapes of the constant tensors and ensure no gpu-cpu sync
+			//will be needed during the execution of the operator.
+			typename TOperatorRegistryRDG<TOperatorType>::OperatorValidateFunc ValidationFn = Registry->OpFindValidation(OpType);
+
+			if (!ValidationFn(AttributeMap, InputTensorTypes))
+			{
+				UE_LOG(LogNNX, Warning, TEXT("Hlsl MLOperatorRegistry failed to validate operator:%s"), *OpType);
+				return false;
+			}
+		}
+
+		return true;
+	}
+};
 
 // NOTE: For now we only have DML on Windows, we should add support for XSX
 #if PLATFORM_WINDOWS
diff --git a/Engine/Plugins/Experimental/NNX/Source/NNXUtils/NNXUtils.Build.cs b/Engine/Plugins/Experimental/NNX/Source/NNXUtils/NNXUtils.Build.cs
--- a/Engine/Plugins/Experimental/NNX/Source/NNXUtils/NNXUtils.Build.cs
+++ b/Engine/Plugins/Experimental/NNX/Source/NNXUtils/NNXUtils.Build.cs
@@ -11,6 +11,12 @@ public class NNXUtils : ModuleRules
 		ShortName = "NNXUtils"; // Shorten to avoid path-too-long errors
 		PCHUsage = ModuleRules.PCHUsageMode.UseExplicitOrSharedPCHs;
 
+		PublicIncludePaths.AddRange(
+			new string[] {
+				Path.GetFullPath(Path.Combine(EngineDirectory, "Plugins/Experimental/NNX/Source/NNXCore/Private")),
+			}
+		);
+
 		PublicDependencyModuleNames.AddRange(
 			new string[]
 			{
@@ -44,6 +50,6 @@ public class NNXUtils : ModuleRules
 				}
 			);
 		}
-		}
-    }
+	}
+}
 
diff --git a/Engine/Plugins/Experimental/NNX/Source/NNXUtils/Private/NNXModelBuilderONNX.cpp b/Engine/Plugins/Experimental/NNX/Source/NNXUtils/Private/NNXModelBuilderONNX.cpp
--- a/Engine/Plugins/Experimental/NNX/Source/NNXUtils/Private/NNXModelBuilderONNX.cpp
+++ b/Engine/Plugins/Experimental/NNX/Source/NNXUtils/Private/NNXModelBuilderONNX.cpp
@@ -2,6 +2,7 @@
 
 #include "NNXCore.h"
 #include "NNXModelBuilder.h"
+#include "NNXRuntimeFormat.h"
 
 #include "NNXThirdPartyWarningDisabler.h"
 NNX_THIRD_PARTY_INCLUDES_START
@@ -274,17 +275,19 @@ private:
 //
 //
 //
-NNXUTILS_API bool CreateONNXModelForOperator(const FString& OperatorName, TConstArrayView<FMLTensorDesc> InInputTensors, TConstArrayView<FMLTensorDesc> InOutputTensors, TArray<uint8>& ModelData)
+NNXUTILS_API bool CreateONNXModelForOperator(const FString& OperatorName, TConstArrayView<FMLTensorDesc> InInputTensors, TConstArrayView<FMLTensorDesc> InOutputTensors, FNNXFormatDesc& Model)
 {
 	FMLAttributeMap EmptyAttributeMap;
-	return CreateONNXModelForOperator(OperatorName, InInputTensors, InOutputTensors, EmptyAttributeMap, ModelData);
+	return CreateONNXModelForOperator(OperatorName, InInputTensors, InOutputTensors, EmptyAttributeMap, Model);
 }
 
 //
 //
 //
-NNXUTILS_API bool CreateONNXModelForOperator(const FString& OperatorName, TConstArrayView<FMLTensorDesc> InInputTensors, TConstArrayView<FMLTensorDesc> InOutputTensors, const FMLAttributeMap& Attributes, TArray<uint8>& ModelData)
+NNXUTILS_API bool CreateONNXModelForOperator(const FString& OperatorName, TConstArrayView<FMLTensorDesc> InInputTensors, TConstArrayView<FMLTensorDesc> InOutputTensors, const FMLAttributeMap& Attributes, FNNXFormatDesc& Model)
 {
+	Model = FNNXFormatDesc{};
+	
 	TUniquePtr<IMLModelBuilder> Builder(CreateONNXModelBuilder());
 
 	Builder->Begin();
@@ -328,8 +331,10 @@ NNXUTILS_API bool CreateONNXModelForOperator(const FString& OperatorName, TConst
 		Builder->AddOperatorAttribute(Op, Attributes.GetName(Idx), Attributes.GetAttributeValue(Idx));
 	}
 
-	Builder->End(ModelData);
-
+	
+	Builder->End(Model.Data);
+	Model.Format = ENNXInferenceFormat::ONNX;
+	
 	return true;
 }
 
diff --git a/Engine/Plugins/Experimental/NNX/Source/NNXUtils/Private/NNXModelOptimizer.cpp b/Engine/Plugins/Experimental/NNX/Source/NNXUtils/Private/NNXModelOptimizer.cpp
--- a/Engine/Plugins/Experimental/NNX/Source/NNXUtils/Private/NNXModelOptimizer.cpp
+++ b/Engine/Plugins/Experimental/NNX/Source/NNXUtils/Private/NNXModelOptimizer.cpp
@@ -21,8 +21,10 @@ NNX_THIRD_PARTY_INCLUDES_START
 #undef check
 #undef TEXT
 #include "core/session/ort_model_optimizer_api.h"
+#include "core/session/onnxruntime_cxx_api.h"
 NNX_THIRD_PARTY_INCLUDES_END
 
+
 #define Print(Format, ...) UE_LOG(LogNNX, Display, Format, __VA_ARGS__)
 
 class FModelGraphPrinter
@@ -220,13 +222,7 @@ public:
 		}
 	}
 
-	static void OnLog(const char* LogMsg)
-	{
-		UE_LOG(LogNNX, Warning, TEXT("%s"), ANSI_TO_TCHAR(LogMsg));
-	}
-
 private:
-
 	TUniquePtr<Ort::IModelGraph>	Graph;
 	TArray<uint8>					Storage;
 };
@@ -266,34 +262,230 @@ EMLTensorDataType GetDataTypeFromGraphTensor(Ort::GraphTensorDataType TensorData
 
 namespace NNX
 {
+class FModelOptimizerBase : public IModelOptimizer
+{
+protected:
+	TArray<TSharedPtr<IModelOptimizerPass>> OptimizationPasses;
+	TArray<TSharedPtr<IModelValidator>> Validators;
+
+	virtual bool ValidateInputModel(const FNNXFormatDesc& InputModel)
+	{
+		if (InputModel.Format != ENNXInferenceFormat::ONNX)
+		{
+			UE_LOG(LogNNX, Warning, TEXT("Optimizer %s is expecting ONNX input format."), *GetName());
+			return false;
+		}
+
+		OrtStatusPtr Status = OrtValidateModelFromMemory(InputModel.Data.GetData(), InputModel.Data.Num());
+		if (Status)
+		{
+			UE_LOG(LogNNX, Warning, TEXT("Input ONNX model is invalid: %s, Model won't be optimized"), ANSI_TO_TCHAR(Ort::GetApi().GetErrorMessage(Status)));
+			return false;
+		}
+		
+		return true;
+	}
+	
+public:
+	virtual void AddOptimizationPass(TSharedPtr<IModelOptimizerPass> ModelOptimizerPass) override
+	{
+		if (ModelOptimizerPass.IsValid())
+		{
+			OptimizationPasses.Add(ModelOptimizerPass);
+		}
+	}
+
+	virtual void AddValidator(TSharedPtr<IModelValidator> ModelValidator) override
+	{
+		if (ModelValidator.IsValid())
+		{
+			Validators.Add(ModelValidator);
+		}
+	}
+
+	bool IsModelValid(const FNNXFormatDesc& ModelToValidate)
+	{
+		bool bIsModelValid = true;
+
+		for (TSharedPtr<IModelValidator>& Validator : Validators)
+		{
+			check(Validator.IsValid());
+			if (!Validator->ValidateModel(ModelToValidate))
+			{
+				UE_LOG(LogNNX, Warning, TEXT("Model validator %s detected an error."), *(Validator->GetName()));
+				bIsModelValid = false;
+			}
+		}
+		return bIsModelValid;
+	}
+
+	bool ApplyAllPassesAndValidations(FNNXFormatDesc& OptimizedModel, const FOptimizerOptionsMap& Options)
+	{
+		if (!IsModelValid(OptimizedModel))
+		{
+			UE_LOG(LogNNX, Warning, TEXT("Model is not valid, skipping optimization passes."));
+			return false;
+		}
+		
+		for (TSharedPtr<IModelOptimizerPass>& Pass : OptimizationPasses)
+		{
+			check(Pass.IsValid());
+			
+			if (!Pass->ApplyPass(OptimizedModel, Options))
+			{
+				UE_LOG(LogNNX, Warning, TEXT("Error while executing model optimisation pass %s."), *(Pass->GetName()));
+				return false;
+			}
+			if (!IsModelValid(OptimizedModel))
+			{
+				UE_LOG(LogNNX, Warning, TEXT("Model validation failed after optimisation pass %s."), *(Pass->GetName()));
+				return false;
+			}
+		}
+
+		return true;
+	}
+
+	bool Optimize(const FNNXFormatDesc& InputModel, FNNXFormatDesc& OptimizedModel, const FOptimizerOptionsMap& Options) override
+	{
+		OptimizedModel = FNNXFormatDesc{};
+		
+		if (!ValidateInputModel(InputModel))
+		{
+			return false;
+		}
+
+		OptimizedModel = InputModel;
+		return ApplyAllPassesAndValidations(OptimizedModel, Options);
+	}
+};
+
+class FONNXModelValidator : public IModelValidator
+{
+public:
+	virtual FString GetName() const
+	{
+		return TEXT("ONNX Model validator");
+	}
+
+	virtual bool ValidateModel(const FNNXFormatDesc& InputModel) const override
+	{
+		FMLRuntimeFormat	Format;
+
+		ENNXInferenceFormat FormatType = InputModel.Format;
+		if (FormatType != ENNXInferenceFormat::ONNX)
+		{
+			UE_LOG(LogNNX, Warning, TEXT("Unsupported format type for validator %s"), *GetName());
+			return false;
+		}
+
+		OrtStatusPtr Status = OrtValidateModelFromMemory(InputModel.Data.GetData(), InputModel.Data.Num());
+		if (Status)
+		{
+			UE_LOG(LogNNX, Warning, TEXT("Failed to validate ONNX model: %s"), ANSI_TO_TCHAR(Ort::GetApi().GetErrorMessage(Status)));
+			return false;
+		}
+
+		return true;
+	}
+};
+
+class FModelOptimizerONNXToONNX : public FModelOptimizerBase
+{
+public:
+	virtual FString GetName() const override
+	{
+		return TEXT("NNXModelOptimizerFromONNXToONNX");
+	}
+
+	//TODO jira 167591: investigate if it is possible to optimize ONNX model and serialize it back as ONNX.
+	//Among other benefits applying L1 could cut down on the model size. Atm this optimizer is a pass trough.
+	//Bonus: Can we get a validator for onnx format?
+};
 
-class FMLModelOptimizerONNXToNNX : public IMLModelOptimizer
+static void OnOrtLog(const char* LogMsg)
+{
+	UE_LOG(LogNNX, Warning, TEXT("%s"), ANSI_TO_TCHAR(LogMsg));
+}
+
+class FModelOptimizerONNXToORT : public FModelOptimizerBase
 {
 public:
+	virtual FString GetName() const override
+	{
+		return TEXT("NNXModelOptimizerONNXToORT");
+	}
 
-	FMLModelOptimizerONNXToNNX()
+	virtual bool Optimize(const FNNXFormatDesc& InputModel, FNNXFormatDesc& OptimizedModel, const FOptimizerOptionsMap& Options) override
 	{
+		OptimizedModel = FNNXFormatDesc{};
+		
+		if (!ValidateInputModel(InputModel))
+		{
+			return false;
+		}
+
+		TUniquePtr<Ort::IModelGraph> Graph = ConvertONNXToORTModelGraph(InputModel.Data);
+
+		//TODO jira 167588: Serialize the Graph to a buffer and store it in OutModel.Data
+		//allowing to initiate a session from ORT format for NN engine supporting this.
+		//Allowing model optimization to be applyed at cooking time.
+		//OptimizedModel = Graph.Serialize;
+		//return ApplyAllPasses(OptimizedModel, Options);
+		return false;
 	}
 
-	virtual bool Optimize(TArrayView<uint8> ONNXData, TArray<uint8>& NNXData) override
+protected:
+	TUniquePtr<Ort::IModelGraph> ConvertONNXToORTModelGraph(TConstArrayView<uint8> ONNXData)
 	{
 		Ort::ModelOptimizeOptions	Opts{};
-		
-		Opts.logCallback = OnLog;
+
+		Opts.logCallback = OnOrtLog;
 
 		TUniquePtr<Ort::IModelGraph> Graph(OrtOptimizeModelFromMemory(ONNXData.GetData(), ONNXData.Num(), Opts));
 		if (!Graph)
 		{
 			UE_LOG(LogNNX, Warning, TEXT("Failed to load ONNX model from memory"));
-			return false;
 		}
 
 		//auto Printer = MakeUnique<FModelGraphPrinter>(Graph.Get());
 		//Printer->Run();
+
+		return Graph;
+	}
+};
+
+class FModelOptimizerONNXToNNXRT : public FModelOptimizerONNXToORT
+{
+public:
+	virtual FString GetName() const override
+	{
+		return TEXT("FNNXModelOptimizerONNXToNNX");
+	}
+	
+	virtual bool Optimize(const FNNXFormatDesc& InputModel, FNNXFormatDesc& OptimizedModel, const FOptimizerOptionsMap& Options) override
+	{
+		OptimizedModel = FNNXFormatDesc{};
 		
-		return BuildNNXFormat(Graph.Get(), NNXData);
+		if (!ValidateInputModel(InputModel))
+		{
+			return false;
+		}
+
+		TUniquePtr<Ort::IModelGraph> Graph = ConvertONNXToORTModelGraph(InputModel.Data);
+		
+		if (!BuildNNXFormat(Graph.Get(), OptimizedModel.Data))
+		{
+			UE_LOG(LogNNX, Warning, TEXT("Error while building ORT ModelGraph."));
+			return false;
+		}
+
+		OptimizedModel.Format = ENNXInferenceFormat::NNXRT;
+
+		return ApplyAllPassesAndValidations(OptimizedModel, Options);
 	}
 
+private:
 	bool BuildNNXFormat(Ort::IModelGraph* Graph, TArray<uint8>& NNXData)
 	{
 		TUniquePtr<NNX::IMLModelBuilder>	Builder(NNX::CreateNNXModelBuilder());
@@ -375,7 +567,7 @@ public:
 				// TODO: Use tensor initializer data
 				if (!Init)
 				{
-					
+					 
 				}
 				
 				auto Tensor =
@@ -414,25 +606,66 @@ public:
 
 		return Builder->End(NNXData);
 	}
-
-	static void OnLog(const char* LogMsg)
-	{
-		UE_LOG(LogNNX, Warning, TEXT("%s"), ANSI_TO_TCHAR(LogMsg));
-	}
 };
 
 /** Create a model optimizer */
-NNXUTILS_API IMLModelOptimizer* CreateModelOptimizer(EMLInferenceFormat InputFormat, EMLInferenceFormat OutputFormat)
+NNXUTILS_API TUniquePtr<IModelOptimizer> CreateModelOptimizer(ENNXInferenceFormat InputFormat, ENNXInferenceFormat OutputFormat)
 {
-	if (InputFormat == EMLInferenceFormat::ONNX)
+	if (InputFormat == ENNXInferenceFormat::ONNX)
 	{
-		if (OutputFormat == EMLInferenceFormat::NNXRT)
+		if (OutputFormat == ENNXInferenceFormat::NNXRT)
+		{
+			return MakeUnique<FModelOptimizerONNXToNNXRT>();
+		}
+		else if (OutputFormat == ENNXInferenceFormat::ONNX)
 		{
-			return new FMLModelOptimizerONNXToNNX();
+			return MakeUnique<FModelOptimizerONNXToONNX>();
+		}
+		else
+		{
+			return MakeUnique<FModelOptimizerONNXToORT>();
 		}
 	}
 
+	//TODO jira 167592: Investigate how to conditionally compile the above,
+	//removing the dependencies to ORT in runtime build (we will then
+	//need a way to run tests without ability to create onnx
+	//model at runtime). One could also make NNXUtils editor only
+	//then the interface for optimization and validation should be in a
+	//different module than the implementation.
+
 	return nullptr;
 }
 
+/** Helper to create an optimized model for a given runtime from ONNX */
+NNXUTILS_API bool CreateRuntimeModelFromONNX(const FNNXFormatDesc& ONNXModel, FNNXFormatDesc& OptimizedModel, FString RuntimeName, const FOptimizerOptionsMap& Options)
+{
+	OptimizedModel = FNNXFormatDesc{};
+	
+	//TODO jira 167594: Register the NNEngine even when CreateInferenceModel can't be called
+	//on the platform because of missing hardware support. This would allow
+	//to optimize a model even when the current hardware cannot run inference on current hardware.
+	NNX::IRuntime* Runtime = NNX::GetRuntime(RuntimeName);
+	if (!Runtime)
+	{
+		UE_LOG(LogNNX, Warning, TEXT("Can't find runtime %s!"), *RuntimeName);
+		return false;
+	}
+
+	TUniquePtr<IModelOptimizer> Optimizer = Runtime->CreateModelOptimizer();
+	if (!Optimizer.IsValid())
+	{
+		UE_LOG(LogNNX, Warning, TEXT("No optimizer available for runtime %s!"), *RuntimeName);
+		return false;
+	}
+
+	if (!Optimizer->Optimize(ONNXModel, OptimizedModel, Options))
+	{
+		UE_LOG(LogNNX, Warning, TEXT("Could not create an optimized model for the required runtime %s. Check the log!"), *RuntimeName);
+		return false;
+	}
+
+	return true;
+}
+
 } // namespace NNX
diff --git a/Engine/Plugins/Experimental/NNX/Source/NNXUtils/Private/NNXUtilsModule.cpp b/Engine/Plugins/Experimental/NNX/Source/NNXUtils/Private/NNXUtilsModule.cpp
--- a/Engine/Plugins/Experimental/NNX/Source/NNXUtils/Private/NNXUtilsModule.cpp
+++ b/Engine/Plugins/Experimental/NNX/Source/NNXUtils/Private/NNXUtilsModule.cpp
@@ -73,24 +73,25 @@ static FAutoConsoleCommand ConsoleCmdNNXOptimizeModel(
 
 			}
 
-			TArray<uint8>	InputData;
+			FNNXFormatDesc ONNXModel;
 			
 			// First try to load model from file, otherwise create a single layer NN with specified
 			// operator name
 			if (!ModelPath.IsEmpty())
 			{
-				if (!FFileHelper::LoadFileToArray(InputData, *ModelPath))
+				if (!FFileHelper::LoadFileToArray(ONNXModel.Data, *ModelPath))
 				{
 					UE_LOG(LogNNX, Warning, TEXT("Failed to load model from file:%s"), *ModelPath);
 					return;
 				}
+				ONNXModel.Format = ENNXInferenceFormat::ONNX;
 			}
 			else if (!OpName.IsEmpty())
 			{
 				FMLTensorDesc InputTensor = FMLTensorDesc::Make(TEXT("in"), { 1, 512 }, EMLTensorDataType::Float);
 				FMLTensorDesc OutputTensor = FMLTensorDesc::Make(TEXT("out"), { 1, 512 }, EMLTensorDataType::Float);
 
-				if (!NNX::CreateONNXModelForOperator(OpName, MakeArrayView(&InputTensor, 1), MakeArrayView(&OutputTensor, 1), InputData))
+				if (!NNX::CreateONNXModelForOperator(OpName, MakeArrayView(&InputTensor, 1), MakeArrayView(&OutputTensor, 1), ONNXModel))
 				{
 					UE_LOG(LogNNX, Warning, TEXT("Failed to create model for operator:%s"), *OpName);
 				}
@@ -100,7 +101,7 @@ static FAutoConsoleCommand ConsoleCmdNNXOptimizeModel(
 				UE_LOG(LogNNX, Warning, TEXT("Invalid arguments"));
 			}
 
-			TUniquePtr<IMLModelOptimizer>	Optimizer(CreateONNXToNNXModelOptimizer());
+			TUniquePtr<IModelOptimizer> Optimizer = CreateONNXToNNXModelOptimizer();
 
 			if (!Optimizer)
 			{
@@ -108,13 +109,14 @@ static FAutoConsoleCommand ConsoleCmdNNXOptimizeModel(
 				return;
 			}
 
-			TArray<uint8>	OutData;
-
 			UE_LOG(LogNNX, Verbose, TEXT("Optimizing model"));
 
-			if (Optimizer->Optimize(InputData, OutData))
+			FOptimizerOptionsMap Options;
+			FNNXFormatDesc OutModel;
+			
+			if (Optimizer->Optimize(ONNXModel, OutModel, Options))
 			{
-				UE_LOG(LogNNX, Display, TEXT("Model optimized:%u bytes"), OutData.Num());
+				UE_LOG(LogNNX, Display, TEXT("Model optimized:%u bytes"), OutModel.Data.Num());
 			}
 			else
 			{
@@ -134,12 +136,11 @@ static FAutoConsoleCommand ConsoleCmdNNXCreateModel(
 
 			if (Args.Num() > 0)
 			{
-				TArray<uint8> ModelData;
-				
+				FNNXFormatDesc ONNXModel;
 				FMLTensorDesc InputTensor = FMLTensorDesc::Make(TEXT("in"), { 1, 512 }, EMLTensorDataType::Float);
 				FMLTensorDesc OutputTensor = FMLTensorDesc::Make(TEXT("out"), { 1, 512 }, EMLTensorDataType::Float);
 
-				if (!NNX::CreateONNXModelForOperator(Args[0], MakeArrayView(&InputTensor, 1), MakeArrayView(&OutputTensor, 1), ModelData))
+				if (!NNX::CreateONNXModelForOperator(Args[0], MakeArrayView(&InputTensor, 1), MakeArrayView(&OutputTensor, 1), ONNXModel))
 				{
 					UE_LOG(LogNNX, Display, TEXT("Failed to create model for operator:%s"), *Args[0]);
 				}
diff --git a/Engine/Plugins/Experimental/NNX/Source/NNXUtils/Public/NNXModelBuilder.h b/Engine/Plugins/Experimental/NNX/Source/NNXUtils/Public/NNXModelBuilder.h
--- a/Engine/Plugins/Experimental/NNX/Source/NNXUtils/Public/NNXModelBuilder.h
+++ b/Engine/Plugins/Experimental/NNX/Source/NNXUtils/Public/NNXModelBuilder.h
@@ -4,6 +4,8 @@
 #include "NNXTypes.h"
 #include "Containers/Array.h"
 
+struct FNNXFormatDesc;
+
 namespace NNX
 {
 
@@ -84,8 +86,8 @@ NNXUTILS_API IMLModelBuilder* CreateONNXModelBuilder(int64 IrVersion = OnnxIrVer
 /**
  * Utility functions to create single layer NN for operator testing with optional attributes
  */
-NNXUTILS_API bool CreateONNXModelForOperator(const FString& OperatorName, TConstArrayView<const FMLTensorDesc> InInputTensors, TConstArrayView<const FMLTensorDesc> InOutputTensors, TArray<uint8>& ModelData);
-NNXUTILS_API bool CreateONNXModelForOperator(const FString& OperatorName, TConstArrayView<const FMLTensorDesc> InInputTensors, TConstArrayView<const FMLTensorDesc> InOutputTensors, const FMLAttributeMap& Attributes, TArray<uint8>& ModelData);
+NNXUTILS_API bool CreateONNXModelForOperator(const FString& OperatorName, TConstArrayView<const FMLTensorDesc> InInputTensors, TConstArrayView<const FMLTensorDesc> InOutputTensors, FNNXFormatDesc& ModelData);
+NNXUTILS_API bool CreateONNXModelForOperator(const FString& OperatorName, TConstArrayView<const FMLTensorDesc> InInputTensors, TConstArrayView<const FMLTensorDesc> InOutputTensors, const FMLAttributeMap& Attributes, FNNXFormatDesc& ModelData);
 
 /**
  * Create an instance of NNX model builder that creates NNX model/format in memory
diff --git a/Engine/Plugins/Experimental/NNX/Source/NNXUtils/Public/NNXModelOptimizer.h b/Engine/Plugins/Experimental/NNX/Source/NNXUtils/Public/NNXModelOptimizer.h
--- a/Engine/Plugins/Experimental/NNX/Source/NNXUtils/Public/NNXModelOptimizer.h
+++ b/Engine/Plugins/Experimental/NNX/Source/NNXUtils/Public/NNXModelOptimizer.h
@@ -1,29 +1,32 @@
 // Copyright Epic Games, Inc. All Rights Reserved.
 #pragma once
 
-#include "NNXInferenceModel.h"
+#include "NNXRuntimeFormat.h"
+#include "NNXModelOptimizerInterface.h"
 
 namespace NNX
 {
+class IModelOptimizer;
+	
+/** Create a model optimizer */
+NNXUTILS_API TUniquePtr<IModelOptimizer> CreateModelOptimizer(ENNXInferenceFormat InputFormat, ENNXInferenceFormat OutputFormat);
 
-/** Base class for ML model optimizers */
-class IMLModelOptimizer
+inline TUniquePtr<IModelOptimizer> CreateONNXToNNXModelOptimizer()
 {
-public:
-
-	virtual ~IMLModelOptimizer() = default;
-
-	/** Optimize model from source to destination format */
-	virtual bool Optimize(TArrayView<uint8> InputModel, TArray<uint8>& OutModel) = 0;
-};
+	return CreateModelOptimizer(ENNXInferenceFormat::ONNX, ENNXInferenceFormat::NNXRT);
+}
 
-/** Create a model optimizer */
-NNXUTILS_API IMLModelOptimizer* CreateModelOptimizer(EMLInferenceFormat InputFormat, EMLInferenceFormat OutputFormat);
+inline TUniquePtr<IModelOptimizer> CreateONNXToORTModelOptimizer()
+{
+	return CreateModelOptimizer(ENNXInferenceFormat::ONNX, ENNXInferenceFormat::ORT);
+}
 
-/** Create a model optimizer from ONNX to NNX */
-inline IMLModelOptimizer* CreateONNXToNNXModelOptimizer()
+inline TUniquePtr<IModelOptimizer> CreateONNXToONNXModelOptimizer()
 {
-	return CreateModelOptimizer(EMLInferenceFormat::ONNX, EMLInferenceFormat::NNXRT);
+	return CreateModelOptimizer(ENNXInferenceFormat::ONNX, ENNXInferenceFormat::ONNX);
 }
 
+/** Helper to create an optimized model for a given runtime from ONNX */
+NNXUTILS_API bool CreateRuntimeModelFromONNX(FNNXFormatDesc& OutputModel, const FNNXFormatDesc& ONNXModel, FString RuntimeName, const FOptimizerOptionsMap& Options);
+
 } // NNX
