NNI - Fix clang compilation error in gather

diff --git a/Engine/Plugins/Experimental/NNX/Source/NNXHlslShaders/Private/NNXGatherCS.cpp b/Engine/Plugins/Experimental/NNX/Source/NNXHlslShaders/Private/NNXGatherCS.cpp
--- a/Engine/Plugins/Experimental/NNX/Source/NNXHlslShaders/Private/NNXGatherCS.cpp
+++ b/Engine/Plugins/Experimental/NNX/Source/NNXHlslShaders/Private/NNXGatherCS.cpp
@@ -4,8 +4,8 @@
 
 namespace UE::NNI::HlslShaders::Internal
 {
-	template <>
-	void TGatherCS<float, int32>::ModifyCompilationEnvironment(const FGlobalShaderPermutationParameters& InParameters, FShaderCompilerEnvironment& OutEnvironment)
+	// template <>
+	void TGatherCS::ModifyCompilationEnvironment(const FGlobalShaderPermutationParameters& InParameters, FShaderCompilerEnvironment& OutEnvironment)
 	{
 		FGlobalShader::ModifyCompilationEnvironment(InParameters, OutEnvironment);
 		OutEnvironment.SetDefine(TEXT("DATA_TYPE"), 0);
@@ -14,8 +14,8 @@ namespace UE::NNI::HlslShaders::Internal
 		OutEnvironment.SetDefine(TEXT("NUM_GROUP_THREADS"), FGatherConstants::NUM_GROUP_THREADS);
 	}
 
-	template <typename DataElementType, typename IndicesElementType>
-	void TGatherCS<DataElementType, IndicesElementType>::FillInParameters(int32 Axis, const NNX::FMLTensorDesc& Data, const NNX::FMLTensorDesc& Indices, FParameters& Parameters)
+	//template <typename DataElementType, typename IndicesElementType>
+	void TGatherCS::FillInParameters(int32 Axis, const NNX::FMLTensorDesc& Data, const NNX::FMLTensorDesc& Indices, FParameters& Parameters)
 	{
 		Parameters.Axis = Axis;
 
@@ -67,14 +67,15 @@ namespace UE::NNI::HlslShaders::Internal
 		}
 	}
 
-	template <typename DataElementType, typename IndicesElementType>
-	FIntVector TGatherCS<DataElementType, IndicesElementType>::GetGroupCount(const FParameters& Parameters)
+	// template <typename DataElementType, typename IndicesElementType>
+	FIntVector TGatherCS::GetGroupCount(const FParameters& Parameters)
 	{
 		return FIntVector(FMath::DivideAndRoundUp(Parameters.OutputSize, FGatherConstants::NUM_GROUP_THREADS), 1, 1);
 	}
 
-	template class TGatherCS<float, int32>;
-	typedef TGatherCS<float, int32> FGatherFloatInt32CS;
-	IMPLEMENT_SHADER_TYPE(template<>, FGatherFloatInt32CS, TEXT("/NNX/GatherOp.usf"), TEXT("main"), SF_Compute);
+	//typedef TGatherCS<float, int32> FGatherFloatInt32CS;
+	//IMPLEMENT_SHADER_TYPE(template<>, TGatherCS, TEXT("/NNX/GatherOp.usf"), TEXT("main"), SF_Compute);
+	IMPLEMENT_GLOBAL_SHADER(TGatherCS, "/NNX/GatherOp.usf", "main", SF_Compute);
+	//template class TGatherCS<float, int32>;
 
 } // UE::NNI::HlslShaders::Internal
\ No newline at end of file
diff --git a/Engine/Plugins/Experimental/NNX/Source/NNXHlslShaders/Public/NNXGatherCS.h b/Engine/Plugins/Experimental/NNX/Source/NNXHlslShaders/Public/NNXGatherCS.h
--- a/Engine/Plugins/Experimental/NNX/Source/NNXHlslShaders/Public/NNXGatherCS.h
+++ b/Engine/Plugins/Experimental/NNX/Source/NNXHlslShaders/Public/NNXGatherCS.h
@@ -16,7 +16,7 @@ namespace UE::NNI::HlslShaders::Internal
 		static const int32 NUM_GROUP_THREADS{ 256 };
 	};
 
-	template <typename DataElementType, typename IndicesElementType>
+	// template <typename DataElementType, typename IndicesElementType>
 	class NNXHLSLSHADERS_API TGatherCS : public FGlobalShader
 	{
 		DECLARE_GLOBAL_SHADER(TGatherCS);
@@ -33,9 +33,9 @@ namespace UE::NNI::HlslShaders::Internal
 			SHADER_PARAMETER(int32, NumIndicesDimensions)
 			SHADER_PARAMETER_ARRAY(FIntVector4, DataStride_IndicesStride_OutputStride, [FGatherConstants::MAX_NUM_DIMENSIONS])
 			SHADER_PARAMETER_ARRAY(FVector4f, OneDivDataStride_OneDivIndicesStride_OneDivOutputStride, [FGatherConstants::MAX_NUM_DIMENSIONS])
-			SHADER_PARAMETER_RDG_BUFFER_SRV(Buffer<DataElementType>, Data)
-			SHADER_PARAMETER_RDG_BUFFER_SRV(Buffer<IndicesElementType>, Indices)
-			SHADER_PARAMETER_RDG_BUFFER_UAV(RWBuffer<DataElementType>, Output)
+			SHADER_PARAMETER_RDG_BUFFER_SRV(Buffer<float>, Data)
+			SHADER_PARAMETER_RDG_BUFFER_SRV(Buffer<int32>, Indices)
+			SHADER_PARAMETER_RDG_BUFFER_UAV(RWBuffer<float>, Output)
 		END_SHADER_PARAMETER_STRUCT()
 
 		static void ModifyCompilationEnvironment(const FGlobalShaderPermutationParameters& InParameters, FShaderCompilerEnvironment& OutEnvironment);
diff --git a/Engine/Plugins/Experimental/NNX/Source/NNXRuntimeRDG/Private/NNXRuntimeHLSLGatherOp.cpp b/Engine/Plugins/Experimental/NNX/Source/NNXRuntimeRDG/Private/NNXRuntimeHLSLGatherOp.cpp
--- a/Engine/Plugins/Experimental/NNX/Source/NNXRuntimeRDG/Private/NNXRuntimeHLSLGatherOp.cpp
+++ b/Engine/Plugins/Experimental/NNX/Source/NNXRuntimeRDG/Private/NNXRuntimeHLSLGatherOp.cpp
@@ -8,8 +8,8 @@ namespace UE::NNI::RuntimeRDG::Hlsl::Private
 {
 	DECLARE_GPU_STAT_NAMED(FNNIOperatorGather, TEXT("NNI.Operator.Hlsl.Gather"));
 
-	template<typename DataElementType, typename IndicesElementType>
-	using TGatherCS = typename UE::NNI::HlslShaders::Internal::TGatherCS<DataElementType, IndicesElementType>;
+	//template<typename DataElementType, typename IndicesElementType>
+	using TGatherCS = typename UE::NNI::HlslShaders::Internal::TGatherCS;
 	using FGatherConstants = UE::NNI::HlslShaders::Internal::FGatherConstants;
 
 	template <typename DataElementType, typename IndicesElementType>
@@ -54,17 +54,17 @@ namespace UE::NNI::RuntimeRDG::Hlsl::Private
 		virtual void Dispatch(FRDGBuilder& GraphBuilder, TArrayView<const NNX::FMLTensorBinding> InInputBindings, TArrayView<const NNX::FMLTensorBinding> OutOutputBindings) override
 		{
 			// Set parameters
-			typename TGatherCS<DataElementType, IndicesElementType>::FParameters* Parameters = GraphBuilder.AllocParameters<typename TGatherCS<DataElementType, IndicesElementType>::FParameters>();
-			TGatherCS<DataElementType, IndicesElementType>::FillInParameters(Axis, Data, Indices, *Parameters);
+			TGatherCS::FParameters* Parameters = GraphBuilder.AllocParameters<TGatherCS::FParameters>();
+			TGatherCS::FillInParameters(Axis, Data, Indices, *Parameters);
 			Parameters->Data = GraphBuilder.CreateSRV(FRDGBufferSRVDesc(InInputBindings[0].Buffer, PF_R32_FLOAT));
 			Parameters->Indices = GraphBuilder.CreateSRV(FRDGBufferSRVDesc(InInputBindings[1].Buffer, PF_R32_FLOAT));
 			Parameters->Output = GraphBuilder.CreateUAV(FRDGBufferUAVDesc(OutOutputBindings[0].Buffer, PF_R32_FLOAT));
 
-			typename TGatherCS<DataElementType, IndicesElementType>::FPermutationDomain PermutationVector;
-			PermutationVector.template Set<typename TGatherCS<DataElementType, IndicesElementType>::FGatherNumOutputDimensions>(Output.Shape.Num());
-			TShaderMapRef<TGatherCS<DataElementType, IndicesElementType>> ComputeShader(GetGlobalShaderMap(GMaxRHIFeatureLevel), PermutationVector);
+			TGatherCS::FPermutationDomain PermutationVector;
+			PermutationVector.Set<TGatherCS::FGatherNumOutputDimensions>(Output.Shape.Num());
+			TShaderMapRef<TGatherCS> ComputeShader(GetGlobalShaderMap(GMaxRHIFeatureLevel), PermutationVector);
 
-			FIntVector ThreadGroupCount = TGatherCS<DataElementType, IndicesElementType>::GetGroupCount(*Parameters);
+			FIntVector ThreadGroupCount = TGatherCS::GetGroupCount(*Parameters);
 
 			RDG_EVENT_SCOPE(GraphBuilder, "NNI.Operator.Hlsl.Gather");
 			RDG_GPU_STAT_SCOPE(GraphBuilder, FNNIOperatorGather);
