Added Conv operator to HLSL runtime (no attributes and only basic input tests as for ConvTranspose).

diff --git a/Engine/Plugins/Experimental/NNX/Source/NNXQA/Resources/NNXQATestDesc.json b/Engine/Plugins/Experimental/NNX/Source/NNXQA/Resources/NNXQATestDesc.json
--- a/Engine/Plugins/Experimental/NNX/Source/NNXQA/Resources/NNXQATestDesc.json
+++ b/Engine/Plugins/Experimental/NNX/Source/NNXQA/Resources/NNXQATestDesc.json
@@ -224,9 +224,10 @@
 			]
 		},
 		{
-			"category": "ConvTranspose",
+			"category": "Convolution",
 			"runtimes": [	{	"name": "NNXRuntimeDml", "skip": true	}	],
 			"targets": [
+				{ "target": "Conv" },
 				{ "target": "ConvTranspose" }
 			]
 		},
@@ -410,7 +411,7 @@
 			]
 		},
 		{
-			"name": "ConvTranspose",
+			"name": "Convolution",
 			"datasets": [
 				{	"inputs": [	{ "shape": [ 1, 1, 1 ] },	{ "shape": [ 1, 1, 1 ] } ],	"outputs": [ { "shape": [ 1, 1, 1 ] } ]},
 				{	"inputs": [	{ "shape": [ 1, 1, 1 ] },	{ "shape": [ 1, 1, 1 ] },	{ "shape": [ 1 ] } ],	"outputs": [ { "shape": [ 1, 1, 1 ] } ]}
diff --git a/Engine/Plugins/Experimental/NNX/Source/NNXRuntimeRDG/Private/Hlsl/NNIRuntimeRDGConv.cpp b/Engine/Plugins/Experimental/NNX/Source/NNXRuntimeRDG/Private/Hlsl/NNIRuntimeRDGConv.cpp
new file mode 100644
--- /dev/null
+++ b/Engine/Plugins/Experimental/NNX/Source/NNXRuntimeRDG/Private/Hlsl/NNIRuntimeRDGConv.cpp
@@ -0,0 +1,130 @@
+// Copyright Epic Games, Inc. All Rights Reserved.
+
+#include "NNIRuntimeRDGConv.h"
+#include "NNXConvCS.h"
+#include "NNXRuntimeHLSLHelper.h"
+
+namespace UE::NNIRuntimeRDG::Private::Hlsl
+{
+	DECLARE_GPU_STAT_NAMED(FMLHLSLOperatorConv, TEXT("NNI.Operator.Hlsl.Conv"));
+
+	/**
+	 * Convolution ML operator implementation
+	 */
+	class FConv : public NNX::FMLOperatorHlsl
+	{
+	public:
+
+		static NNX::FMLOperatorHlsl* Create()
+		{
+			return new FConv();
+		}
+
+		virtual ~FConv() = default;
+
+	private:
+
+		FConv() {}
+
+		NNX::FMLTensorDesc Input;
+		NNX::FMLTensorDesc Weights;
+		NNX::FMLTensorDesc Bias;
+		NNX::FMLTensorDesc Output;
+
+		int NumDimensions;
+		bool HasBias;
+
+		// Hard coded parameters, until we accept them from JSON
+		int32 Group = 1;
+		TArray<int32> InputShape;
+		TArray<int32> WeightsShape;
+		EConvAutoPad AutoPad = EConvAutoPad::VALID;
+		TArray<int32> Dilations = {1};
+		TArray<int32> Strides = {1};
+		TArray<int32> Pads = {0, 0};
+
+	public:
+
+		virtual bool Initialize(TArrayView<const NNX::FMLTensorDesc> InputTensors, TArrayView<const NNX::FMLTensorDesc> OutputTensors, const FMLAttributeMap& Attributes) override
+		{
+			check(InputTensors.Num() >= 2 && InputTensors.Num() <= 3);
+			check(OutputTensors.Num() == 1);
+
+			Input = InputTensors[0];
+			Weights = InputTensors[1];
+			Output = OutputTensors[0];
+
+			check(Input.Shape.Num() > 2);
+			check(Weights.Shape.Num()  == Input.Shape.Num());
+			check(Output.Shape.Num() == Input.Shape.Num());
+
+			if (InputTensors.Num() == 3) {
+				HasBias = true;
+				Bias = InputTensors[2];
+			}
+			else {
+				HasBias = false;
+			}
+
+			auto MakeShape = [](const NNX::FMLTensorDesc& InputDesc) {
+				TArray<int32> Result;
+				for (int32 i = 0; i < (int32)InputDesc.Shape.Num(); i++) {
+					Result.Add(InputDesc.Shape[i]);
+				}
+				return Result;
+			};
+			InputShape = MakeShape(Input);
+			WeightsShape = MakeShape(Weights);
+
+			NumDimensions = Input.Shape.Num() - 2;
+
+			return true;
+		}
+
+		virtual void Dispatch(FRDGBuilder& GraphBuilder, TArrayView<const NNX::FMLTensorBinding> InInputBindings, TArrayView<const NNX::FMLTensorBinding> OutOutputBindings) override
+		{
+			constexpr EConvAlgorithm Algorithm = EConvAlgorithm::SharedMemory;
+			constexpr EConvGroupSize GroupSize = EConvGroupSize::Size256;
+
+			TArray<int32> OutputShape = FMLConvCS::GetOutputShape(InputShape, WeightsShape, AutoPad, Dilations, Strides, Pads);
+
+			// Set parameters
+			FMLConvCS::FParameters* Params = GraphBuilder.AllocParameters<FMLConvCS::FParameters>();
+			FMLConvCS::FillInParameters(GroupSize, InputShape, WeightsShape, HasBias, AutoPad, Group, Dilations,Strides, Pads, *Params);
+			Params->X = GraphBuilder.CreateSRV(FRDGBufferSRVDesc(InInputBindings[0].Buffer, PF_R32_FLOAT));
+			Params->W = GraphBuilder.CreateSRV(FRDGBufferSRVDesc(InInputBindings[1].Buffer, PF_R32_FLOAT));
+			if (InInputBindings.Num() == 3) {
+				Params->B = GraphBuilder.CreateSRV(FRDGBufferSRVDesc(InInputBindings[2].Buffer, PF_R32_FLOAT));
+			}
+			Params->Y = GraphBuilder.CreateUAV(FRDGBufferUAVDesc(OutOutputBindings[0].Buffer, PF_R32_FLOAT));
+
+			FMLConvCS::FPermutationDomain PermutationVector;
+
+			PermutationVector.Set<FMLConvCS::FConvAlgorithm>(Algorithm);
+			PermutationVector.Set<FMLConvCS::FConvGroupSize>(GroupSize);
+			PermutationVector.Set<FMLConvCS::FConvNumDimensions>(NumDimensions);
+			PermutationVector.Set<FMLConvCS::FConvNumReadsPerThread>(FMLConvCS::GetNumReadsPerThread(GroupSize, WeightsShape, Dilations, Strides));
+			PermutationVector.Set<FMLConvCS::FConvHasB>(HasBias);
+			TShaderMapRef<FMLConvCS> ComputeShader(GetGlobalShaderMap(GMaxRHIFeatureLevel), PermutationVector);
+
+			RDG_EVENT_SCOPE(GraphBuilder, "FML.HLSL.Operator.Conv");
+			RDG_GPU_STAT_SCOPE(GraphBuilder, FMLHLSLOperatorConv);
+
+			FComputeShaderUtils::AddPass(
+				GraphBuilder,
+				RDG_EVENT_NAME("FML.HLSL.Operator.Conv.Dispatch"),
+				ERDGPassFlags::Compute | ERDGPassFlags::NeverCull,
+				ComputeShader,
+				Params,
+				FMLConvCS::GetGroupCount(OutputShape, FMLConvCS::GetGroupShape(GroupSize, NumDimensions)));		
+		}
+	};
+
+	bool RegisterConvOperator(NNX::FMLOperatorRegistryHlsl& Registry)
+	{
+		Registry.OpAdd(TEXT("Conv"), FConv::Create);
+
+		return true;
+	}
+
+} // NNX
diff --git a/Engine/Plugins/Experimental/NNX/Source/NNXRuntimeRDG/Private/Hlsl/NNIRuntimeRDGConv.h b/Engine/Plugins/Experimental/NNX/Source/NNXRuntimeRDG/Private/Hlsl/NNIRuntimeRDGConv.h
new file mode 100644
--- /dev/null
+++ b/Engine/Plugins/Experimental/NNX/Source/NNXRuntimeRDG/Private/Hlsl/NNIRuntimeRDGConv.h
@@ -0,0 +1,11 @@
+// Copyright Epic Games, Inc. All Rights Reserved.
+
+#pragma once
+
+#include "NNXRuntimeRDG.h"
+#include "NNXRuntimeHLSLOp.h"
+
+namespace UE::NNIRuntimeRDG::Private::Hlsl
+{
+	bool RegisterConvOperator(NNX::FMLOperatorRegistryHlsl& Registry);
+} // UE::NNIRuntimeRDG::Private::Hlsl
diff --git a/Engine/Plugins/Experimental/NNX/Source/NNXRuntimeRDG/Private/NNXRuntimeHLSL.cpp b/Engine/Plugins/Experimental/NNX/Source/NNXRuntimeRDG/Private/NNXRuntimeHLSL.cpp
--- a/Engine/Plugins/Experimental/NNX/Source/NNXRuntimeRDG/Private/NNXRuntimeHLSL.cpp
+++ b/Engine/Plugins/Experimental/NNX/Source/NNXRuntimeRDG/Private/NNXRuntimeHLSL.cpp
@@ -1,16 +1,17 @@
 // Copyright Epic Games, Inc. All Rights Reserved.
 
+#include "Hlsl/NNIRuntimeRDGConv.h"
 #include "NNXElementWiseBinaryCS.h"
 #include "NNXInferenceModel.h"
-#include "NNXRuntimeFormat.h"
 #include "NNXModelOptimizer.h"
+#include "NNXRuntimeFormat.h"
+#include "NNXRuntimeHLSLConvTransposeOp.h"
 #include "NNXRuntimeHLSLElementWiseBinaryOps.h"
 #include "NNXRuntimeHLSLElementWiseUnaryOps.h"
 #include "NNXRuntimeHLSLElementWiseVariadicOps.h"
 #include "NNXRuntimeHLSLGemmOp.h"
-#include "NNXRuntimeHLSLConvTransposeOp.h"
-#include "NNXRuntimeHLSLMatMulOp.h"
 #include "NNXRuntimeHLSLHelper.h"
+#include "NNXRuntimeHLSLMatMulOp.h"
 #include "NNXRuntimeRDG.h"
 
 namespace NNX
@@ -164,6 +165,7 @@ public:
 		RegisterElementWiseBinaryOperators(*registry);
 		RegisterElementWiseVariadicOperators(*registry);
 		RegisterGemmOperator(*registry);
+		UE::NNIRuntimeRDG::Private::Hlsl::RegisterConvOperator(*registry);
 		RegisterConvTransposeOperator(*registry);
 		RegisterMatMulOperator(*registry);
 
