NNX Quick cleanup

diff --git a/Engine/Plugins/Experimental/NNX/Source/NNXCore/Private/NNXInferenceModel.cpp b/Engine/Plugins/Experimental/NNX/Source/NNXCore/Private/NNXInferenceModel.cpp
--- a/Engine/Plugins/Experimental/NNX/Source/NNXCore/Private/NNXInferenceModel.cpp
+++ b/Engine/Plugins/Experimental/NNX/Source/NNXCore/Private/NNXInferenceModel.cpp
@@ -2,7 +2,7 @@
 
 #include "NNXInferenceModel.h"
 
-UMLInferenceModel* UMLInferenceModel::CreateFromFormatDesc(const FNNXFormatDesc& FormatDesc)
+UMLInferenceModel* UMLInferenceModel::CreateFromFormatDesc(const FNNIModelRaw& FormatDesc)
 {
 	UMLInferenceModel* Model = NewObject<UMLInferenceModel>((UObject*)GetTransientPackage(), UMLInferenceModel::StaticClass());
 
@@ -16,7 +16,7 @@ UMLInferenceModel* UMLInferenceModel::CreateFromFormatDesc(const FNNXFormatDesc&
 	return Model;
 }
 
-const FNNXFormatDesc& UMLInferenceModel::GetFormatDesc() const
+const FNNIModelRaw& UMLInferenceModel::GetFormatDesc() const
 {
 	return FormatDesc;
 }
\ No newline at end of file
diff --git a/Engine/Plugins/Experimental/NNX/Source/NNXCore/Private/NNXModelOptimizerInterface.h b/Engine/Plugins/Experimental/NNX/Source/NNXCore/Private/NNXModelOptimizerInterface.h
--- a/Engine/Plugins/Experimental/NNX/Source/NNXCore/Private/NNXModelOptimizerInterface.h
+++ b/Engine/Plugins/Experimental/NNX/Source/NNXCore/Private/NNXModelOptimizerInterface.h
@@ -16,7 +16,7 @@ class IModelValidator
 public:
 	virtual ~IModelValidator() = default;
 	virtual FString GetName() const = 0;
-	virtual bool ValidateModel(const FNNXFormatDesc& InputModel) const = 0;
+	virtual bool ValidateModel(const FNNIModelRaw& InputModel) const = 0;
 };
 
 /** Interface class for NNX model optimizer pass */
@@ -27,14 +27,13 @@ public:
 	virtual FString GetName() const = 0;
 
 	//Optimize the model in place, potentially changing the format
-	virtual bool ApplyPass(FNNXFormatDesc& Model, const FOptimizerOptionsMap& Options) const = 0;
+	virtual bool ApplyPass(FNNIModelRaw& Model, const FOptimizerOptionsMap& Options) const = 0;
 };
 
 /** Interface class for NNX model optimizer */
 class IModelOptimizer
 {
 public:
-
 	virtual ~IModelOptimizer() = default;
 	virtual FString GetName() const = 0;
 
@@ -45,7 +44,7 @@ public:
 	virtual void AddValidator(TSharedPtr<IModelValidator>) = 0;
 	
 	//Apply all passes and validators to the input model, produce an optimized model potentially in a different format
-	virtual bool Optimize(const FNNXFormatDesc& InputModel, FNNXFormatDesc& OutModel, const FOptimizerOptionsMap& Options) = 0;
+	virtual bool Optimize(const FNNIModelRaw& InputModel, FNNIModelRaw& OutModel, const FOptimizerOptionsMap& Options) = 0;
 };
 
 } // NNX
diff --git a/Engine/Plugins/Experimental/NNX/Source/NNXCore/Public/NNXInferenceModel.h b/Engine/Plugins/Experimental/NNX/Source/NNXCore/Public/NNXInferenceModel.h
--- a/Engine/Plugins/Experimental/NNX/Source/NNXCore/Public/NNXInferenceModel.h
+++ b/Engine/Plugins/Experimental/NNX/Source/NNXCore/Public/NNXInferenceModel.h
@@ -13,13 +13,13 @@ class NNXCORE_API UMLInferenceModel : public UObject
 
 public:
 
-	static UMLInferenceModel* CreateFromFormatDesc(const FNNXFormatDesc& Model);
+	static UMLInferenceModel* CreateFromFormatDesc(const FNNIModelRaw& Model);
 	
 	UMLInferenceModel() = default;
 	
-	const FNNXFormatDesc& GetFormatDesc() const;
+	const FNNIModelRaw& GetFormatDesc() const;
 
 private:
 	UPROPERTY(EditAnywhere, Category = "Neural Network Inference")
-	FNNXFormatDesc FormatDesc;
+	FNNIModelRaw FormatDesc;
 };
diff --git a/Engine/Plugins/Experimental/NNX/Source/NNXCore/Public/NNXRuntimeFormat.h b/Engine/Plugins/Experimental/NNX/Source/NNXCore/Public/NNXRuntimeFormat.h
--- a/Engine/Plugins/Experimental/NNX/Source/NNXCore/Public/NNXRuntimeFormat.h
+++ b/Engine/Plugins/Experimental/NNX/Source/NNXCore/Public/NNXRuntimeFormat.h
@@ -27,7 +27,7 @@ enum class ENNXInferenceFormat : uint8
 };
 
 USTRUCT()
-struct FNNXFormatDesc
+struct FNNIModelRaw
 {
 	GENERATED_USTRUCT_BODY()
 	
@@ -36,8 +36,6 @@ struct FNNXFormatDesc
 	
 	UPROPERTY(VisibleAnywhere, Category = "Neural Network Inference")
 	ENNXInferenceFormat	Format { ENNXInferenceFormat::Invalid };
-
-	FNNXFormatDesc() = default;
 };
 
 // Required by LoadModel() when loading operators in HLSL and DirectML runtime
diff --git a/Engine/Plugins/Experimental/NNX/Source/NNXQA/Private/NNXQAParametricTest.cpp b/Engine/Plugins/Experimental/NNX/Source/NNXQA/Private/NNXQAParametricTest.cpp
--- a/Engine/Plugins/Experimental/NNX/Source/NNXQA/Private/NNXQAParametricTest.cpp
+++ b/Engine/Plugins/Experimental/NNX/Source/NNXQA/Private/NNXQAParametricTest.cpp
@@ -390,7 +390,7 @@ namespace Test
 	
 	static bool RunParametricTest(FTests::FTestSetup& TestSetup, const FString& RuntimeFilter)
 	{
-		FNNXFormatDesc ONNXModel;
+		FNNIModelRaw ONNXModel;
 
 		if (TestSetup.IsModelTest)
 		{
diff --git a/Engine/Plugins/Experimental/NNX/Source/NNXQA/Private/NNXQAUtils.cpp b/Engine/Plugins/Experimental/NNX/Source/NNXQA/Private/NNXQAUtils.cpp
--- a/Engine/Plugins/Experimental/NNX/Source/NNXQA/Private/NNXQAUtils.cpp
+++ b/Engine/Plugins/Experimental/NNX/Source/NNXQA/Private/NNXQAUtils.cpp
@@ -325,14 +325,14 @@ namespace Test
 		return FMath::Sin((float)(ElementIndex + IndexOffsetBetweenTensor * TensorIndex));
 	}
 
-	static int RunTestInference(const FNNXFormatDesc& ONNXModelData, NNX::IRuntime* Runtime, 
+	static int RunTestInference(const FNNIModelRaw& ONNXModelData, NNX::IRuntime* Runtime, 
 		TArray<NNX::FMLTensorDesc>& OutOutputTensorsDesc, TArray<TArray<char>>& OutOutputMemBuffers)
 	{
 		OutOutputMemBuffers.Empty();
 		OutOutputTensorsDesc.Empty();
 
 		FOptimizerOptionsMap Options;
-		FNNXFormatDesc RuntimeModelData;
+		FNNIModelRaw RuntimeModelData;
 		TUniquePtr<IModelOptimizer> Optimizer = Runtime->CreateModelOptimizer();
 		
 		if (!Optimizer || !Optimizer->Optimize(ONNXModelData, RuntimeModelData, Options))
@@ -385,7 +385,7 @@ namespace Test
 		return ReturnValue;
 	}
 
-	bool CompareONNXModelInferenceAcrossRuntimes(const FNNXFormatDesc& ONNXModel, const FTests::FTestSetup& TestSetup, const FString& RuntimeFilter)
+	bool CompareONNXModelInferenceAcrossRuntimes(const FNNIModelRaw& ONNXModel, const FTests::FTestSetup& TestSetup, const FString& RuntimeFilter)
 	{
 		FString CurrentPlatform = UGameplayStatics::GetPlatformName();
 		if (TestSetup.AutomationExcludedPlatform.Contains(CurrentPlatform))
diff --git a/Engine/Plugins/Experimental/NNX/Source/NNXQA/Private/NNXQAUtils.h b/Engine/Plugins/Experimental/NNX/Source/NNXQA/Private/NNXQAUtils.h
--- a/Engine/Plugins/Experimental/NNX/Source/NNXQA/Private/NNXQAUtils.h
+++ b/Engine/Plugins/Experimental/NNX/Source/NNXQA/Private/NNXQAUtils.h
@@ -5,7 +5,7 @@
 #include "Containers/UnrealString.h"
 #include "NNXRuntime.h"
 
-struct FNNXFormatDesc;
+struct FNNIModelRaw;
 
 namespace NNX 
 {
@@ -60,7 +60,7 @@ namespace Test
 		TArray<FTestSetup> TestSetups;
 	};
 
-	bool CompareONNXModelInferenceAcrossRuntimes(const FNNXFormatDesc& ONNXModel, const FTests::FTestSetup& TestSetup, const FString& RuntimeFilter = TEXT(""));
+	bool CompareONNXModelInferenceAcrossRuntimes(const FNNIModelRaw& ONNXModel, const FTests::FTestSetup& TestSetup, const FString& RuntimeFilter = TEXT(""));
 
 	FString FMLTensorDescToString(const FMLTensorDesc& TensorDesc);
 	FString TensorToString(const FMLTensorDesc& TensorDesc, const TArray<char>& TensorData);
diff --git a/Engine/Plugins/Experimental/NNX/Source/NNXRuntimeRDG/Private/NNXRuntimeRDG.cpp b/Engine/Plugins/Experimental/NNX/Source/NNXRuntimeRDG/Private/NNXRuntimeRDG.cpp
--- a/Engine/Plugins/Experimental/NNX/Source/NNXRuntimeRDG/Private/NNXRuntimeRDG.cpp
+++ b/Engine/Plugins/Experimental/NNX/Source/NNXRuntimeRDG/Private/NNXRuntimeRDG.cpp
@@ -176,7 +176,7 @@ FMLInferenceModelRDG::~FMLInferenceModelRDG()
 //
 //
 //
-bool FMLInferenceModelRDG::LoadModel(const FNNXFormatDesc& InModel, FMLRuntimeFormat& Format)
+bool FMLInferenceModelRDG::LoadModel(const FNNIModelRaw& InModel, FMLRuntimeFormat& Format)
 {
 	ENNXInferenceFormat FormatType = InModel.Format;
 
diff --git a/Engine/Plugins/Experimental/NNX/Source/NNXRuntimeRDG/Private/NNXRuntimeRDG.h b/Engine/Plugins/Experimental/NNX/Source/NNXRuntimeRDG/Private/NNXRuntimeRDG.h
--- a/Engine/Plugins/Experimental/NNX/Source/NNXRuntimeRDG/Private/NNXRuntimeRDG.h
+++ b/Engine/Plugins/Experimental/NNX/Source/NNXRuntimeRDG/Private/NNXRuntimeRDG.h
@@ -26,7 +26,7 @@ END_SHADER_PARAMETER_STRUCT()
 
 class FRDGBuilder;
 struct FMLRuntimeFormat;
-struct FNNXFormatDesc;
+struct FNNIModelRaw;
 
 namespace NNX
 {
@@ -67,7 +67,7 @@ protected:
 
 	FMLInferenceModelRDG();
 
-	bool LoadModel(const FNNXFormatDesc& InModel, FMLRuntimeFormat& Format);
+	bool LoadModel(const FNNIModelRaw& InModel, FMLRuntimeFormat& Format);
 
 	int SetTensors(FRDGBuilder& GraphBuilder, FMLTensorBindingArray& OutRDGBindings, FMLIntArray& OutIndices, TArrayView<const FMLTensorBinding> InBindings, TArrayView<const FMLTensorDesc> InTensors);
 
@@ -203,7 +203,7 @@ public:
 		return TEXT("RDG Model validator");
 	}
 	
-	virtual bool ValidateModel(const FNNXFormatDesc& InputModel) const override
+	virtual bool ValidateModel(const FNNIModelRaw& InputModel) const override
 	{
 		FMLRuntimeFormat	Format;
 
@@ -220,12 +220,6 @@ public:
 		TOperatorRegistryRDG<TOperatorType>* Registry = TOperatorRegistryRDG<TOperatorType>::Get();
 		check(Registry != nullptr);
 
-		if (Format.Operators.Num() > 1)
-		{
-			UE_LOG(LogNNX, Warning, TEXT("Failed to validate model, currently only single layer models are supported"));
-			return false;
-		}
-		
 		// HACK: This works only for single layer networks
 		// we should get the output tensor type from intermediate tensor 
 		// (can we get it from model definition?, otherwise will need
diff --git a/Engine/Plugins/Experimental/NNX/Source/NNXUtils/Private/NNXModelBuilderONNX.cpp b/Engine/Plugins/Experimental/NNX/Source/NNXUtils/Private/NNXModelBuilderONNX.cpp
--- a/Engine/Plugins/Experimental/NNX/Source/NNXUtils/Private/NNXModelBuilderONNX.cpp
+++ b/Engine/Plugins/Experimental/NNX/Source/NNXUtils/Private/NNXModelBuilderONNX.cpp
@@ -275,7 +275,7 @@ private:
 //
 //
 //
-NNXUTILS_API bool CreateONNXModelForOperator(const FString& OperatorName, TConstArrayView<FMLTensorDesc> InInputTensors, TConstArrayView<FMLTensorDesc> InOutputTensors, FNNXFormatDesc& Model)
+NNXUTILS_API bool CreateONNXModelForOperator(const FString& OperatorName, TConstArrayView<FMLTensorDesc> InInputTensors, TConstArrayView<FMLTensorDesc> InOutputTensors, FNNIModelRaw& Model)
 {
 	FMLAttributeMap EmptyAttributeMap;
 	return CreateONNXModelForOperator(OperatorName, InInputTensors, InOutputTensors, EmptyAttributeMap, Model);
@@ -284,9 +284,9 @@ NNXUTILS_API bool CreateONNXModelForOperator(const FString& OperatorName, TConst
 //
 //
 //
-NNXUTILS_API bool CreateONNXModelForOperator(const FString& OperatorName, TConstArrayView<FMLTensorDesc> InInputTensors, TConstArrayView<FMLTensorDesc> InOutputTensors, const FMLAttributeMap& Attributes, FNNXFormatDesc& Model)
+NNXUTILS_API bool CreateONNXModelForOperator(const FString& OperatorName, TConstArrayView<FMLTensorDesc> InInputTensors, TConstArrayView<FMLTensorDesc> InOutputTensors, const FMLAttributeMap& Attributes, FNNIModelRaw& Model)
 {
-	Model = FNNXFormatDesc{};
+	Model = FNNIModelRaw{};
 	
 	TUniquePtr<IMLModelBuilder> Builder(CreateONNXModelBuilder());
 
diff --git a/Engine/Plugins/Experimental/NNX/Source/NNXUtils/Private/NNXModelOptimizer.cpp b/Engine/Plugins/Experimental/NNX/Source/NNXUtils/Private/NNXModelOptimizer.cpp
--- a/Engine/Plugins/Experimental/NNX/Source/NNXUtils/Private/NNXModelOptimizer.cpp
+++ b/Engine/Plugins/Experimental/NNX/Source/NNXUtils/Private/NNXModelOptimizer.cpp
@@ -268,7 +268,7 @@ protected:
 	TArray<TSharedPtr<IModelOptimizerPass>> OptimizationPasses;
 	TArray<TSharedPtr<IModelValidator>> Validators;
 
-	virtual bool ValidateInputModel(const FNNXFormatDesc& InputModel)
+	virtual bool ValidateInputModel(const FNNIModelRaw& InputModel)
 	{
 		if (InputModel.Format != ENNXInferenceFormat::ONNX)
 		{
@@ -303,7 +303,7 @@ public:
 		}
 	}
 
-	bool IsModelValid(const FNNXFormatDesc& ModelToValidate)
+	bool IsModelValid(const FNNIModelRaw& ModelToValidate)
 	{
 		bool bIsModelValid = true;
 
@@ -319,7 +319,7 @@ public:
 		return bIsModelValid;
 	}
 
-	bool ApplyAllPassesAndValidations(FNNXFormatDesc& OptimizedModel, const FOptimizerOptionsMap& Options)
+	bool ApplyAllPassesAndValidations(FNNIModelRaw& OptimizedModel, const FOptimizerOptionsMap& Options)
 	{
 		if (!IsModelValid(OptimizedModel))
 		{
@@ -346,9 +346,9 @@ public:
 		return true;
 	}
 
-	bool Optimize(const FNNXFormatDesc& InputModel, FNNXFormatDesc& OptimizedModel, const FOptimizerOptionsMap& Options) override
+	bool Optimize(const FNNIModelRaw& InputModel, FNNIModelRaw& OptimizedModel, const FOptimizerOptionsMap& Options) override
 	{
-		OptimizedModel = FNNXFormatDesc{};
+		OptimizedModel = FNNIModelRaw{};
 		
 		if (!ValidateInputModel(InputModel))
 		{
@@ -368,7 +368,7 @@ public:
 		return TEXT("ONNX Model validator");
 	}
 
-	virtual bool ValidateModel(const FNNXFormatDesc& InputModel) const override
+	virtual bool ValidateModel(const FNNIModelRaw& InputModel) const override
 	{
 		FMLRuntimeFormat	Format;
 
@@ -416,9 +416,9 @@ public:
 		return TEXT("NNXModelOptimizerONNXToORT");
 	}
 
-	virtual bool Optimize(const FNNXFormatDesc& InputModel, FNNXFormatDesc& OptimizedModel, const FOptimizerOptionsMap& Options) override
+	virtual bool Optimize(const FNNIModelRaw& InputModel, FNNIModelRaw& OptimizedModel, const FOptimizerOptionsMap& Options) override
 	{
-		OptimizedModel = FNNXFormatDesc{};
+		OptimizedModel = FNNIModelRaw{};
 		
 		if (!ValidateInputModel(InputModel))
 		{
@@ -463,9 +463,9 @@ public:
 		return TEXT("FNNXModelOptimizerONNXToNNX");
 	}
 	
-	virtual bool Optimize(const FNNXFormatDesc& InputModel, FNNXFormatDesc& OptimizedModel, const FOptimizerOptionsMap& Options) override
+	virtual bool Optimize(const FNNIModelRaw& InputModel, FNNIModelRaw& OptimizedModel, const FOptimizerOptionsMap& Options) override
 	{
-		OptimizedModel = FNNXFormatDesc{};
+		OptimizedModel = FNNIModelRaw{};
 		
 		if (!ValidateInputModel(InputModel))
 		{
@@ -638,9 +638,9 @@ NNXUTILS_API TUniquePtr<IModelOptimizer> CreateModelOptimizer(ENNXInferenceForma
 }
 
 /** Helper to create an optimized model for a given runtime from ONNX */
-NNXUTILS_API bool CreateRuntimeModelFromONNX(const FNNXFormatDesc& ONNXModel, FNNXFormatDesc& OptimizedModel, FString RuntimeName, const FOptimizerOptionsMap& Options)
+NNXUTILS_API bool CreateRuntimeModelFromONNX(const FNNIModelRaw& ONNXModel, FNNIModelRaw& OptimizedModel, FString RuntimeName, const FOptimizerOptionsMap& Options)
 {
-	OptimizedModel = FNNXFormatDesc{};
+	OptimizedModel = FNNIModelRaw{};
 	
 	//TODO jira 167594: Register the NNEngine even when CreateInferenceModel can't be called
 	//on the platform because of missing hardware support. This would allow
diff --git a/Engine/Plugins/Experimental/NNX/Source/NNXUtils/Private/NNXUtilsModule.cpp b/Engine/Plugins/Experimental/NNX/Source/NNXUtils/Private/NNXUtilsModule.cpp
--- a/Engine/Plugins/Experimental/NNX/Source/NNXUtils/Private/NNXUtilsModule.cpp
+++ b/Engine/Plugins/Experimental/NNX/Source/NNXUtils/Private/NNXUtilsModule.cpp
@@ -73,7 +73,7 @@ static FAutoConsoleCommand ConsoleCmdNNXOptimizeModel(
 
 			}
 
-			FNNXFormatDesc ONNXModel;
+			FNNIModelRaw ONNXModel;
 			
 			// First try to load model from file, otherwise create a single layer NN with specified
 			// operator name
@@ -112,7 +112,7 @@ static FAutoConsoleCommand ConsoleCmdNNXOptimizeModel(
 			UE_LOG(LogNNX, Verbose, TEXT("Optimizing model"));
 
 			FOptimizerOptionsMap Options;
-			FNNXFormatDesc OutModel;
+			FNNIModelRaw OutModel;
 			
 			if (Optimizer->Optimize(ONNXModel, OutModel, Options))
 			{
@@ -136,7 +136,7 @@ static FAutoConsoleCommand ConsoleCmdNNXCreateModel(
 
 			if (Args.Num() > 0)
 			{
-				FNNXFormatDesc ONNXModel;
+				FNNIModelRaw ONNXModel;
 				FMLTensorDesc InputTensor = FMLTensorDesc::Make(TEXT("in"), { 1, 512 }, EMLTensorDataType::Float);
 				FMLTensorDesc OutputTensor = FMLTensorDesc::Make(TEXT("out"), { 1, 512 }, EMLTensorDataType::Float);
 
diff --git a/Engine/Plugins/Experimental/NNX/Source/NNXUtils/Public/NNXModelBuilder.h b/Engine/Plugins/Experimental/NNX/Source/NNXUtils/Public/NNXModelBuilder.h
--- a/Engine/Plugins/Experimental/NNX/Source/NNXUtils/Public/NNXModelBuilder.h
+++ b/Engine/Plugins/Experimental/NNX/Source/NNXUtils/Public/NNXModelBuilder.h
@@ -4,7 +4,7 @@
 #include "NNXTypes.h"
 #include "Containers/Array.h"
 
-struct FNNXFormatDesc;
+struct FNNIModelRaw;
 
 namespace NNX
 {
@@ -86,8 +86,8 @@ NNXUTILS_API IMLModelBuilder* CreateONNXModelBuilder(int64 IrVersion = OnnxIrVer
 /**
  * Utility functions to create single layer NN for operator testing with optional attributes
  */
-NNXUTILS_API bool CreateONNXModelForOperator(const FString& OperatorName, TConstArrayView<const FMLTensorDesc> InInputTensors, TConstArrayView<const FMLTensorDesc> InOutputTensors, FNNXFormatDesc& ModelData);
-NNXUTILS_API bool CreateONNXModelForOperator(const FString& OperatorName, TConstArrayView<const FMLTensorDesc> InInputTensors, TConstArrayView<const FMLTensorDesc> InOutputTensors, const FMLAttributeMap& Attributes, FNNXFormatDesc& ModelData);
+NNXUTILS_API bool CreateONNXModelForOperator(const FString& OperatorName, TConstArrayView<const FMLTensorDesc> InInputTensors, TConstArrayView<const FMLTensorDesc> InOutputTensors, FNNIModelRaw& ModelData);
+NNXUTILS_API bool CreateONNXModelForOperator(const FString& OperatorName, TConstArrayView<const FMLTensorDesc> InInputTensors, TConstArrayView<const FMLTensorDesc> InOutputTensors, const FMLAttributeMap& Attributes, FNNIModelRaw& ModelData);
 
 /**
  * Create an instance of NNX model builder that creates NNX model/format in memory
diff --git a/Engine/Plugins/Experimental/NNX/Source/NNXUtils/Public/NNXModelOptimizer.h b/Engine/Plugins/Experimental/NNX/Source/NNXUtils/Public/NNXModelOptimizer.h
--- a/Engine/Plugins/Experimental/NNX/Source/NNXUtils/Public/NNXModelOptimizer.h
+++ b/Engine/Plugins/Experimental/NNX/Source/NNXUtils/Public/NNXModelOptimizer.h
@@ -27,6 +27,6 @@ inline TUniquePtr<IModelOptimizer> CreateONNXToONNXModelOptimizer()
 }
 
 /** Helper to create an optimized model for a given runtime from ONNX */
-NNXUTILS_API bool CreateRuntimeModelFromONNX(FNNXFormatDesc& OutputModel, const FNNXFormatDesc& ONNXModel, FString RuntimeName, const FOptimizerOptionsMap& Options);
+NNXUTILS_API bool CreateRuntimeModelFromONNX(FNNIModelRaw& OutputModel, const FNNIModelRaw& ONNXModel, FString RuntimeName, const FOptimizerOptionsMap& Options);
 
 } // NNX
