ConvCS cleanup.

diff --git a/Engine/Plugins/Experimental/NNX/Source/NNXHlslShaders/Private/NNXConvCS.cpp b/Engine/Plugins/Experimental/NNX/Source/NNXHlslShaders/Private/NNXConvCS.cpp
--- a/Engine/Plugins/Experimental/NNX/Source/NNXHlslShaders/Private/NNXConvCS.cpp
+++ b/Engine/Plugins/Experimental/NNX/Source/NNXHlslShaders/Private/NNXConvCS.cpp
@@ -6,7 +6,7 @@
 namespace
 {
 
-TArray<int32> GetXBlockShape(TArray<int32> GroupShape, TArray<int32> WShape, TArray<int32> Dilations, TArray<int32> Strides)
+TArray<int32> GetXBlockShape(TArrayView<const int32> GroupShape, TArrayView<const uint32> WShape, TArrayView<const int32> Dilations, TArrayView<const int32> Strides)
 {
 	check(WShape.Num() > 2);
 	check(GroupShape.Num() == WShape.Num() - 2);
@@ -23,7 +23,7 @@ TArray<int32> GetXBlockShape(TArray<int32> GroupShape, TArray<int32> WShape, TAr
 	return Result;
 }
 
-TArray<int32> GetPadding(TArray<int32> XShape, TArray<int32> WShape, EConvAutoPad AutoPad, TArray<int32> Dilations, TArray<int32> Strides, TArray<int32> Pads)
+TArray<int32> GetPadding(TArrayView<const uint32> XShape, TArrayView<const uint32> WShape, EConvAutoPad AutoPad, TArrayView<const int32> Dilations, TArrayView<const int32> Strides, TArrayView<const int32> Pads)
 {
 	check(XShape.Num() > 2);
 	check(WShape.Num() == XShape.Num());
@@ -37,7 +37,7 @@ TArray<int32> GetPadding(TArray<int32> XShape, TArray<int32> WShape, EConvAutoPa
 
 	if (AutoPad == EConvAutoPad::NOTSET)
 	{
-		return Pads;
+		return TArray<int32>{Pads};
 	}
 	else if (AutoPad == EConvAutoPad::VALID)
 	{
@@ -84,7 +84,7 @@ int32 GetNumThreadsPerGroup(EConvGroupSize GroupSize)
 	return NumThreadsPerGroup;
 }
 
-TArray<int32> GetGridShape(TArray<int32> YShape, TArray<int32> GroupShape)
+TArray<int32> GetGridShape(TArrayView<const int32> YShape, TArrayView<const int32> GroupShape)
 {
 	check(YShape.Num() > 2);
 	check(YShape.Num() == (GroupShape.Num() + 2));
@@ -107,7 +107,7 @@ void FMLConvCS::ModifyCompilationEnvironment(const FGlobalShaderPermutationParam
 	OutEnvironment.SetDefine(TEXT("MAX_NUM_DIMENSIONS"), FConvConstants::MAX_NUM_DIMENSIONS);
 }
 
-TArray<int32> FMLConvCS::GetOutputShape(TArray<int32> XShape, TArray<int32> WShape, EConvAutoPad AutoPad, TArray<int32> Dilations, TArray<int32> Strides, TArray<int32> Pads)
+TArray<int32> FMLConvCS::GetOutputShape(TArrayView<const uint32> XShape, TArrayView<const uint32> WShape, EConvAutoPad AutoPad, TArrayView<const int32> Dilations, TArrayView<const int32> Strides, TArrayView<const int32> Pads)
 {
 	check(XShape.Num() > 2);
 	check(WShape.Num() == XShape.Num());
@@ -135,7 +135,8 @@ TArray<int32> FMLConvCS::GetOutputShape(TArray<int32> XShape, TArray<int32> WSha
 	return Result;
 }
 
-void FMLConvCS::FillInParameters(EConvGroupSize GroupSize, TArray<int32> XShape, TArray<int32> WShape, bool HasB, EConvAutoPad AutoPad, int Group, TArray<int32> Dilations, TArray<int32> Strides, TArray<int32> Pads, FMLConvCS::FParameters& Parameters)
+void FMLConvCS::FillInParameters(EConvGroupSize GroupSize, TArrayView<const uint32> XShape, TArrayView<const uint32> WShape, bool HasB,
+		EConvAutoPad AutoPad, int Group, TArrayView<const int32> Dilations, TArrayView<const int32> Strides, TArrayView<const int32> Pads, FMLConvCS::FParameters& Parameters)
 {
 	check(XShape.Num() > 2);
 	check(WShape.Num() == XShape.Num());
@@ -188,9 +189,9 @@ void FMLConvCS::FillInParameters(EConvGroupSize GroupSize, TArray<int32> XShape,
 
 	Parameters.XBlockSize = XBlockSize;
 
-	Parameters.NumChannelsPerBatch = FMath::Min((int32)((float)GroupThreadStride / (float)WChannelSize), WShape[1]);
+	Parameters.NumChannelsPerBatch = FMath::Min((int32)((float)GroupThreadStride / (float)WChannelSize), (int32)WShape[1]);
 	check(Parameters.NumChannelsPerBatch > 0)
-	Parameters.NumChannelBatches = FMath::DivideAndRoundUp(WShape[1], Parameters.NumChannelsPerBatch);
+	Parameters.NumChannelBatches = FMath::DivideAndRoundUp((int32)WShape[1], Parameters.NumChannelsPerBatch);
 	
 	Parameters.WOutputKernelStride = WShape[1] * WChannelSize;
 	Parameters.WChannelBatchSize = Parameters.NumChannelsPerBatch * WChannelSize;
@@ -199,7 +200,7 @@ void FMLConvCS::FillInParameters(EConvGroupSize GroupSize, TArray<int32> XShape,
 	Parameters.GroupsDivM = (float)Group / (float)WShape[0];
 }
 
-int32 FMLConvCS::GetNumReadsPerThread(EConvGroupSize GroupSize, TArray<int32> WShape, TArray<int32> Dilations, TArray<int32> Strides)
+int32 FMLConvCS::GetNumReadsPerThread(EConvGroupSize GroupSize, TArrayView<const uint32> WShape, TArrayView<const int32> Dilations, TArrayView<const int32> Strides)
 {
 	check(WShape.Num() > 2);
 	check(Dilations.Num() == 0 || Dilations.Num() == WShape.Num() - 2);
@@ -251,7 +252,7 @@ TArray<int32> FMLConvCS::GetGroupShape(EConvGroupSize GroupSize, int32 NumDimens
 	return Result;
 }
 
-FIntVector FMLConvCS::GetGroupCount(TArray<int32> YShape, TArray<int32> GroupShape)
+FIntVector FMLConvCS::GetGroupCount(TArrayView<const int32> YShape, TArrayView<const int32> GroupShape)
 {
 	check(YShape.Num() > 2);
 	check(YShape.Num() == (GroupShape.Num() + 2));
@@ -265,7 +266,7 @@ FIntVector FMLConvCS::GetGroupCount(TArray<int32> YShape, TArray<int32> GroupSha
 	return FIntVector(ThreadGroupCountValueX, YShape[1], YShape[0]);
 }
 
-EConvGroupSize FMLConvCS::GetMinimalGroupSize(TArray<int32> WShape)
+EConvGroupSize FMLConvCS::GetMinimalGroupSize(TArrayView<const int32> WShape)
 {
 	int32 NumDimensions = WShape.Num() - 2;
 	int32 WChannelSize = 1;
diff --git a/Engine/Plugins/Experimental/NNX/Source/NNXHlslShaders/Public/NNXConvCS.h b/Engine/Plugins/Experimental/NNX/Source/NNXHlslShaders/Public/NNXConvCS.h
--- a/Engine/Plugins/Experimental/NNX/Source/NNXHlslShaders/Public/NNXConvCS.h
+++ b/Engine/Plugins/Experimental/NNX/Source/NNXHlslShaders/Public/NNXConvCS.h
@@ -77,11 +77,12 @@ public:
 	static void ModifyCompilationEnvironment(const FGlobalShaderPermutationParameters& InParameters, FShaderCompilerEnvironment& OutEnvironment);
 
 public:
-	static TArray<int32> GetOutputShape(TArray<int32> XShape, TArray<int32> WShape, EConvAutoPad AutoPad, TArray<int32> Dilations, TArray<int32> Strides, TArray<int32> Pads);
+	static TArray<int32> GetOutputShape(TArrayView<const uint32> XShape, TArrayView<const uint32> WShape, EConvAutoPad AutoPad, TArrayView<const int32> Dilations, TArrayView<const int32> Strides, TArrayView<const int32> Pads);
 
-	static void FillInParameters(EConvGroupSize GroupSize, TArray<int32> XShape, TArray<int32> WShape, bool HasB, EConvAutoPad AutoPad, int Group, TArray<int32> Dilations, TArray<int32> Strides, TArray<int32> Pads, FMLConvCS::FParameters& Parameters);
+	static void FillInParameters(EConvGroupSize GroupSize, TArrayView<const uint32> XShape, TArrayView<const uint32> WShape, bool HasB,
+			EConvAutoPad AutoPad, int Group, TArrayView<const int32> Dilations, TArrayView<const int32> Strides, TArrayView<const int32> Pads, FMLConvCS::FParameters& Parameters);
 
-	static int32 GetNumReadsPerThread(EConvGroupSize GroupSize, TArray<int32> WShape, TArray<int32> Dilations, TArray<int32> Strides);
+	static int32 GetNumReadsPerThread(EConvGroupSize GroupSize, TArrayView<const uint32> WShape, TArrayView<const int32> Dilations, TArrayView<const int32> Strides);
 
 	/**
 	 * @brief Computes the group shape such that all dimension have roughly equal sizes.
@@ -99,7 +100,7 @@ public:
 	 * @param YShape The shape of a single thread group as computed by GetGroupShape()
 	 * @return FIntVector The number of thread groups to instantiate. z corresponds to the batch and y to the output kernel.
 	 */
-	static FIntVector GetGroupCount(TArray<int32> YShape, TArray<int32> GroupShape);
+	static FIntVector GetGroupCount(TArrayView<const int32> YShape, TArrayView<const int32> GroupShape);
 
-	static EConvGroupSize GetMinimalGroupSize(TArray<int32> WShape);
+	static EConvGroupSize GetMinimalGroupSize(TArrayView<const int32> WShape);
 };
\ No newline at end of file
diff --git a/Engine/Plugins/Experimental/NNX/Source/NNXRuntimeRDG/Private/Hlsl/NNIRuntimeRDGConv.cpp b/Engine/Plugins/Experimental/NNX/Source/NNXRuntimeRDG/Private/Hlsl/NNIRuntimeRDGConv.cpp
--- a/Engine/Plugins/Experimental/NNX/Source/NNXRuntimeRDG/Private/Hlsl/NNIRuntimeRDGConv.cpp
+++ b/Engine/Plugins/Experimental/NNX/Source/NNXRuntimeRDG/Private/Hlsl/NNIRuntimeRDGConv.cpp
@@ -36,8 +36,6 @@ namespace UE::NNIRuntimeRDG::Private::Hlsl
 
 		// Hard coded parameters, until we accept them from JSON
 		int32 Group = 1;
-		TArray<int32> InputShape;
-		TArray<int32> WeightsShape;
 		EConvAutoPad AutoPad = EConvAutoPad::VALID;
 		TArray<int32> Dilations = {1};
 		TArray<int32> Strides = {1};
@@ -66,16 +64,6 @@ namespace UE::NNIRuntimeRDG::Private::Hlsl
 				HasBias = false;
 			}
 
-			auto MakeShape = [](const NNX::FMLTensorDesc& InputDesc) {
-				TArray<int32> Result;
-				for (int32 i = 0; i < (int32)InputDesc.Shape.Num(); i++) {
-					Result.Add(InputDesc.Shape[i]);
-				}
-				return Result;
-			};
-			InputShape = MakeShape(Input);
-			WeightsShape = MakeShape(Weights);
-
 			NumDimensions = Input.Shape.Num() - 2;
 
 			return true;
@@ -86,11 +74,11 @@ namespace UE::NNIRuntimeRDG::Private::Hlsl
 			constexpr EConvAlgorithm Algorithm = EConvAlgorithm::SharedMemory;
 			constexpr EConvGroupSize GroupSize = EConvGroupSize::Size256;
 
-			TArray<int32> OutputShape = FMLConvCS::GetOutputShape(InputShape, WeightsShape, AutoPad, Dilations, Strides, Pads);
+			TArray<int32> OutputShape = FMLConvCS::GetOutputShape(Input.Shape, Weights.Shape, AutoPad, Dilations, Strides, Pads);
 
 			// Set parameters
 			FMLConvCS::FParameters* Params = GraphBuilder.AllocParameters<FMLConvCS::FParameters>();
-			FMLConvCS::FillInParameters(GroupSize, InputShape, WeightsShape, HasBias, AutoPad, Group, Dilations,Strides, Pads, *Params);
+			FMLConvCS::FillInParameters(GroupSize, Input.Shape, Weights.Shape, HasBias, AutoPad, Group, Dilations,Strides, Pads, *Params);
 			Params->X = GraphBuilder.CreateSRV(FRDGBufferSRVDesc(InInputBindings[0].Buffer, PF_R32_FLOAT));
 			Params->W = GraphBuilder.CreateSRV(FRDGBufferSRVDesc(InInputBindings[1].Buffer, PF_R32_FLOAT));
 			if (InInputBindings.Num() == 3) {
@@ -103,7 +91,7 @@ namespace UE::NNIRuntimeRDG::Private::Hlsl
 			PermutationVector.Set<FMLConvCS::FConvAlgorithm>(Algorithm);
 			PermutationVector.Set<FMLConvCS::FConvGroupSize>(GroupSize);
 			PermutationVector.Set<FMLConvCS::FConvNumDimensions>(NumDimensions);
-			PermutationVector.Set<FMLConvCS::FConvNumReadsPerThread>(FMLConvCS::GetNumReadsPerThread(GroupSize, WeightsShape, Dilations, Strides));
+			PermutationVector.Set<FMLConvCS::FConvNumReadsPerThread>(FMLConvCS::GetNumReadsPerThread(GroupSize, Weights.Shape, Dilations, Strides));
 			PermutationVector.Set<FMLConvCS::FConvHasB>(HasBias);
 			TShaderMapRef<FMLConvCS> ComputeShader(GetGlobalShaderMap(GMaxRHIFeatureLevel), PermutationVector);
 
@@ -116,7 +104,7 @@ namespace UE::NNIRuntimeRDG::Private::Hlsl
 				ERDGPassFlags::Compute | ERDGPassFlags::NeverCull,
 				ComputeShader,
 				Params,
-				FMLConvCS::GetGroupCount(OutputShape, FMLConvCS::GetGroupShape(GroupSize, NumDimensions)));		
+				FMLConvCS::GetGroupCount(OutputShape, FMLConvCS::GetGroupShape(GroupSize, NumDimensions)));
 		}
 	};
 
