NNI - Add namespaces to elementwise ops

diff --git a/Engine/Plugins/Experimental/NNX/Source/NNXHlslShaders/Internal/NNIHlslShadersElementWiseBinaryCS.h b/Engine/Plugins/Experimental/NNX/Source/NNXHlslShaders/Internal/NNIHlslShadersElementWiseBinaryCS.h
new file mode 100644
--- /dev/null
+++ b/Engine/Plugins/Experimental/NNX/Source/NNXHlslShaders/Internal/NNIHlslShadersElementWiseBinaryCS.h
@@ -0,0 +1,46 @@
+// Copyright Epic Games, Inc. All Rights Reserved.
+
+#pragma once
+
+#include "GlobalShader.h"
+#include "ShaderParameterUtils.h"
+#include "RenderGraphUtils.h"
+#include "NNXOperator.h"
+
+namespace UE::NNIHlslShaders::Internal
+{
+	class FElementWiseBinaryConstants
+	{
+	public:
+
+		static const int32 MAX_NUM_DIMENSIONS{ 8 };
+		static const int32 NUM_GROUP_THREADS{ 256 };
+	};
+
+	class NNXHLSLSHADERS_API TElementWiseBinaryCS : public FGlobalShader
+	{
+		DECLARE_GLOBAL_SHADER(TElementWiseBinaryCS);
+		SHADER_USE_PARAMETER_STRUCT(TElementWiseBinaryCS, FGlobalShader)
+
+		class FOperatorType : SHADER_PERMUTATION_ENUM_CLASS("OP_TYPENAME", EMLElementWiseBinaryOperatorType);
+		class FBinaryNumDimensions : SHADER_PERMUTATION_RANGE_INT("NUM_DIMENSIONS", 1, FElementWiseBinaryConstants::MAX_NUM_DIMENSIONS);
+		using FPermutationDomain = TShaderPermutationDomain<FOperatorType, FBinaryNumDimensions>;
+
+	public:
+
+		BEGIN_SHADER_PARAMETER_STRUCT(FParameters, )
+			SHADER_PARAMETER_RDG_BUFFER_SRV(Buffer<float>, LHSInput)
+			SHADER_PARAMETER_RDG_BUFFER_SRV(Buffer<float>, RHSInput)
+			SHADER_PARAMETER_RDG_BUFFER_UAV(RWBuffer<float>, Output)
+			SHADER_PARAMETER_ARRAY(FUintVector4, TensorInfo, [FElementWiseBinaryConstants::MAX_NUM_DIMENSIONS])
+			SHADER_PARAMETER(uint32, Num)
+			SHADER_PARAMETER(uint32, ThreadCountX)
+		END_SHADER_PARAMETER_STRUCT()
+
+		static void ModifyCompilationEnvironment(const FGlobalShaderPermutationParameters& InParameters, FShaderCompilerEnvironment& OutEnvironment);
+
+	private:
+
+		static const FString GetOpFunc(EMLElementWiseBinaryOperatorType OpType);
+	};
+} // UE::NNIHlslShaders::Internal
diff --git a/Engine/Plugins/Experimental/NNX/Source/NNXHlslShaders/Internal/NNIHlslShadersElementWiseUnaryCS.h b/Engine/Plugins/Experimental/NNX/Source/NNXHlslShaders/Internal/NNIHlslShadersElementWiseUnaryCS.h
new file mode 100644
--- /dev/null
+++ b/Engine/Plugins/Experimental/NNX/Source/NNXHlslShaders/Internal/NNIHlslShadersElementWiseUnaryCS.h
@@ -0,0 +1,45 @@
+// Copyright Epic Games, Inc. All Rights Reserved.
+
+#pragma once
+
+#include "GlobalShader.h"
+#include "ShaderParameterUtils.h"
+#include "RenderGraphUtils.h"
+#include "NNXOperator.h"
+
+namespace UE::NNIHlslShaders::Internal
+{
+	class FElementWiseUnaryConstants
+	{
+	public:
+
+		static const int32 NUM_GROUP_THREADS{ 256 };
+	};
+
+	class NNXHLSLSHADERS_API TElementWiseUnaryCS : public FGlobalShader
+	{
+		DECLARE_GLOBAL_SHADER(TElementWiseUnaryCS);
+		SHADER_USE_PARAMETER_STRUCT(TElementWiseUnaryCS, FGlobalShader)
+
+		class FOperatorType : SHADER_PERMUTATION_ENUM_CLASS("OP_TYPENAME", EMLElementWiseUnaryOperatorType);
+		using FPermutationDomain = TShaderPermutationDomain<FOperatorType>;
+
+	public:
+
+		BEGIN_SHADER_PARAMETER_STRUCT(FParameters, )
+			SHADER_PARAMETER_RDG_BUFFER_SRV(Buffer<float>, Input)
+			SHADER_PARAMETER_RDG_BUFFER_UAV(RWBuffer<float>, Output)
+			SHADER_PARAMETER(uint32, Num)
+			SHADER_PARAMETER(uint32, ThreadCountX)
+			SHADER_PARAMETER(float, Alpha)
+			SHADER_PARAMETER(float, Beta)
+			SHADER_PARAMETER(float, Gamma)
+		END_SHADER_PARAMETER_STRUCT()
+
+		static void ModifyCompilationEnvironment(const FGlobalShaderPermutationParameters& InParameters, FShaderCompilerEnvironment& OutEnvironment);
+
+	private:
+
+		static const FString GetOpFunc(EMLElementWiseUnaryOperatorType OpType);
+	};
+} // UE::NNIHlslShaders::Internal
diff --git a/Engine/Plugins/Experimental/NNX/Source/NNXHlslShaders/Internal/NNIHlslShadersElementWiseVariadicCS.h b/Engine/Plugins/Experimental/NNX/Source/NNXHlslShaders/Internal/NNIHlslShadersElementWiseVariadicCS.h
new file mode 100644
--- /dev/null
+++ b/Engine/Plugins/Experimental/NNX/Source/NNXHlslShaders/Internal/NNIHlslShadersElementWiseVariadicCS.h
@@ -0,0 +1,54 @@
+// Copyright Epic Games, Inc. All Rights Reserved.
+
+#pragma once
+
+#include "GlobalShader.h"
+#include "ShaderParameterUtils.h"
+#include "RenderGraphUtils.h"
+#include "NNXOperator.h"
+
+namespace UE::NNIHlslShaders::Internal
+{
+	class FElementWiseVariadicConstants
+	{
+	public:
+
+		static const int32 MAX_NUM_DIMENSIONS{ 8 };
+		static const uint32 NUM_GROUP_THREADS{ 256 };
+		static const uint32 MAX_NUM_INPUT{ 4 };
+	};
+
+	class NNXHLSLSHADERS_API TElementWiseVariadicCS : public FGlobalShader
+	{
+		DECLARE_GLOBAL_SHADER(TElementWiseVariadicCS);
+		SHADER_USE_PARAMETER_STRUCT(TElementWiseVariadicCS, FGlobalShader)
+
+		class FOperatorType : SHADER_PERMUTATION_ENUM_CLASS("OP_TYPENAME", EMLElementWiseVariadicOperatorType);
+		class FApplyScale : SHADER_PERMUTATION_BOOL("APPLYSCALE");
+		class FOutputAsInput : SHADER_PERMUTATION_BOOL("OUTPUTASINPUT");
+		class FNumInput : SHADER_PERMUTATION_RANGE_INT("NUMINPUT", 1, FElementWiseVariadicConstants::MAX_NUM_INPUT);
+		class FVariadicNumDimensions : SHADER_PERMUTATION_RANGE_INT("NUM_DIMENSIONS", 1, FElementWiseVariadicConstants::MAX_NUM_DIMENSIONS);
+		using FPermutationDomain = TShaderPermutationDomain<FOperatorType, FApplyScale, FOutputAsInput, FNumInput, FVariadicNumDimensions>;
+
+	public:
+
+		BEGIN_SHADER_PARAMETER_STRUCT(FParameters, )
+			SHADER_PARAMETER_RDG_BUFFER_SRV(Buffer<float>, Input0)
+			SHADER_PARAMETER_RDG_BUFFER_SRV(Buffer<float>, Input1)
+			SHADER_PARAMETER_RDG_BUFFER_SRV(Buffer<float>, Input2)
+			SHADER_PARAMETER_RDG_BUFFER_SRV(Buffer<float>, Input3)
+			SHADER_PARAMETER_RDG_BUFFER_UAV(RWBuffer<float>, Output)
+			SHADER_PARAMETER_ARRAY(FUintVector4, InputTensorInfo, [FElementWiseVariadicConstants::MAX_NUM_DIMENSIONS])
+			SHADER_PARAMETER_ARRAY(FUintVector4, OutputTensorInfo, [FElementWiseVariadicConstants::MAX_NUM_DIMENSIONS])
+			SHADER_PARAMETER(uint32, Num)
+			SHADER_PARAMETER(uint32, ThreadCountX)
+			SHADER_PARAMETER(float, Scale)
+		END_SHADER_PARAMETER_STRUCT()
+
+		static void ModifyCompilationEnvironment(const FGlobalShaderPermutationParameters& InParameters, FShaderCompilerEnvironment& OutEnvironment);
+
+	private:
+
+		static const FString GetOpFunc(EMLElementWiseVariadicOperatorType OpType);
+	};
+} // UE::NNIHlslShaders::Internal
\ No newline at end of file
diff --git a/Engine/Plugins/Experimental/NNX/Source/NNXHlslShaders/Private/NNIHlslShadersElementWiseBinaryCS.cpp b/Engine/Plugins/Experimental/NNX/Source/NNXHlslShaders/Private/NNIHlslShadersElementWiseBinaryCS.cpp
new file mode 100644
--- /dev/null
+++ b/Engine/Plugins/Experimental/NNX/Source/NNXHlslShaders/Private/NNIHlslShadersElementWiseBinaryCS.cpp
@@ -0,0 +1,59 @@
+// Copyright Epic Games, Inc. All Rights Reserved.
+
+#include "NNIHlslShadersElementWiseBinaryCS.h"
+#include "NNXCore.h"
+
+namespace UE::NNIHlslShaders::Internal
+{
+	void TElementWiseBinaryCS::ModifyCompilationEnvironment(const FGlobalShaderPermutationParameters& InParameters, FShaderCompilerEnvironment& OutEnvironment)
+	{
+		FGlobalShader::ModifyCompilationEnvironment(InParameters, OutEnvironment);
+
+		OutEnvironment.SetDefine(TEXT("THREADGROUP_SIZE_X"), FElementWiseBinaryConstants::NUM_GROUP_THREADS);
+
+		FPermutationDomain PermutationVector(InParameters.PermutationId);
+
+		const FString OpFunc = GetOpFunc(PermutationVector.Get<FOperatorType>());
+
+		OutEnvironment.SetDefine(TEXT("ELEMENTWISE_OP(X,Y)"), *OpFunc);
+	}
+
+	const FString TElementWiseBinaryCS::GetOpFunc(EMLElementWiseBinaryOperatorType OpType)
+	{
+		FString OpTable[(int32) EMLElementWiseBinaryOperatorType::MAX];
+
+		for (int32 Idx = 0; Idx < (int32)EMLElementWiseBinaryOperatorType::MAX; ++Idx)
+		{
+			OpTable[Idx] = FString("");
+		}
+
+#define OP(OpName, OpFunc) OpTable[(int32) EMLElementWiseBinaryOperatorType::OpName] = OpFunc
+		OP(Add,              TEXT("((X)+(Y))"));
+		//OP(And,            TEXT("((X)&&(Y))"));
+		OP(Div,              TEXT("((X)/(Y))"));
+		//OP(Equal,          TEXT("((X)==(Y))"));
+		//OP(Greater,        TEXT("((X)>(Y))"));
+		//OP(GreaterOrEqual, TEXT("((X)>=(Y))"));
+		//OP(Less,           TEXT("((X)<(Y))"));
+		//OP(LessOrEqual,    TEXT("((X)<(Y))"));
+		OP(Mod,              TEXT("((X)%(Y))"));
+		OP(Mul,              TEXT("((X)*(Y))"));
+		//OP(Or,             TEXT("((X)||(Y))"));
+		OP(Prelu,            TEXT("prelu(X,Y)"));
+		OP(Pow,              TEXT("safe_pow(X,Y)"));
+		OP(Sub,              TEXT("((X)-(Y))"));
+		//OP(Or,             TEXT("((X)^=(Y))"));
+#undef OP
+
+		FString OpFunc = OpTable[(int32) OpType];
+
+		if (OpFunc == "")
+		{
+			UE_LOG(LogNNX, Warning, TEXT("Undefined ElementWise Binary operator name for operator:%d"), OpType);
+		}
+
+		return OpFunc;
+	}
+
+	IMPLEMENT_GLOBAL_SHADER(TElementWiseBinaryCS, "/NNX/ElementWiseBinaryOp.usf", "ElementWiseBinaryOp", SF_Compute);
+} // UE::NNIHlslShaders::Internal
\ No newline at end of file
diff --git a/Engine/Plugins/Experimental/NNX/Source/NNXHlslShaders/Private/NNIHlslShadersElementWiseUnaryCS.cpp b/Engine/Plugins/Experimental/NNX/Source/NNXHlslShaders/Private/NNIHlslShadersElementWiseUnaryCS.cpp
new file mode 100644
--- /dev/null
+++ b/Engine/Plugins/Experimental/NNX/Source/NNXHlslShaders/Private/NNIHlslShadersElementWiseUnaryCS.cpp
@@ -0,0 +1,82 @@
+// Copyright Epic Games, Inc. All Rights Reserved.
+
+#include "NNIHlslShadersElementWiseUnaryCS.h"
+#include "NNXCore.h"
+
+namespace UE::NNIHlslShaders::Internal
+{
+	void TElementWiseUnaryCS::ModifyCompilationEnvironment(const FGlobalShaderPermutationParameters& InParameters, FShaderCompilerEnvironment& OutEnvironment)
+	{
+		FGlobalShader::ModifyCompilationEnvironment(InParameters, OutEnvironment);
+
+		OutEnvironment.SetDefine(TEXT("THREADGROUP_SIZE_X"), FElementWiseUnaryConstants::NUM_GROUP_THREADS);
+
+		FPermutationDomain PermutationVector(InParameters.PermutationId);
+
+		const FString OpFunc = GetOpFunc(PermutationVector.Get<FOperatorType>());
+
+		OutEnvironment.SetDefine(TEXT("ELEMENTWISE_OP(X)"), *OpFunc);
+	}
+
+	const FString TElementWiseUnaryCS::GetOpFunc(EMLElementWiseUnaryOperatorType OpType)
+	{
+		FString OpTable[(int32) EMLElementWiseUnaryOperatorType::MAX];
+
+		for (int32 Idx = 0; Idx < (int32) EMLElementWiseUnaryOperatorType::MAX; ++Idx)
+		{
+			OpTable[Idx] = FString("");
+		}
+
+#define OP(OpName, OpFunc) OpTable[(int32) EMLElementWiseUnaryOperatorType::OpName] = OpFunc
+		OP(Abs,         TEXT("abs(X)"));
+		OP(Acos,        TEXT("acos(X)"));
+		OP(Acosh,       TEXT("acosh(X)"));
+		OP(Asin,        TEXT("asin(X)"));
+		OP(Asinh,       TEXT("asinh(X)"));
+		OP(Atan,        TEXT("atan(X)"));
+		OP(Atanh,       TEXT("atanh(X)"));
+		//OP(BitShift,  TEXT("bitshift(X)"));
+		//OP(Cast,      TEXT("cast(X)"));
+		OP(Ceil,        TEXT("ceil(X)"));
+		//OP(Clip,      TEXT("clip(X)"));
+		OP(Cos,         TEXT("cos(X)"));
+		OP(Cosh,        TEXT("cosh(X)"));
+		OP(Elu,         TEXT("elu(X)"));
+		OP(Erf,         TEXT("erf(X)"));
+		OP(Exp,         TEXT("exp(X)"));
+		OP(Floor,       TEXT("floor(X)"));
+		OP(IsInf,       TEXT("isinf(X)"));
+		OP(IsNan,       TEXT("isnan(X)"));//TODO check shader there is a warning saying input can neither be Nan on PC FXC
+		OP(HardSigmoid, TEXT("hardSigmoid(X)"));
+		OP(HardSwish,   TEXT("hardSwish(X)"));
+		OP(LeakyRelu,   TEXT("leakyRelu(X)"));
+		OP(Log,         TEXT("log(X)"));
+		OP(Neg,         TEXT("-(X)"));
+		//OP(Not,       TEXT("not(X)"));
+		OP(Reciprocal,  TEXT("1 / (X)"));
+		OP(Relu,        TEXT("relu(X)"));
+		OP(Round,       TEXT("round(X)"));
+		OP(Selu,        TEXT("selu(X)"));
+		OP(Sigmoid,     TEXT("sigmoid(X)"));
+		OP(Sign,        TEXT("sign(X)"));
+		OP(Sin,         TEXT("sin(X)"));
+		OP(Sinh,        TEXT("sinh(X)"));
+		OP(Softplus,    TEXT("softplus(X)"));
+		OP(Softsign,    TEXT("softsign(X)"));
+		OP(Sqrt,        TEXT("sqrt(X)"));
+		OP(Tan,         TEXT("tan(X)"));
+		OP(Tanh,        TEXT("tanh(X)"));
+#undef OP
+
+		FString OpFunc = OpTable[(int32) OpType];
+
+		if (OpFunc == "")
+		{
+			UE_LOG(LogNNX, Warning, TEXT("Undefined ElementWise Unary operator name for operator:%d"), OpType);
+		}
+
+		return OpFunc;
+	}
+
+	IMPLEMENT_GLOBAL_SHADER(TElementWiseUnaryCS, "/NNX/ElementWiseOp.usf", "ElementWiseOp", SF_Compute);
+} // UE::NNIHlslShaders::Internal
diff --git a/Engine/Plugins/Experimental/NNX/Source/NNXHlslShaders/Private/NNIHlslShadersElementWiseVariadicCS.cpp b/Engine/Plugins/Experimental/NNX/Source/NNXHlslShaders/Private/NNIHlslShadersElementWiseVariadicCS.cpp
new file mode 100644
--- /dev/null
+++ b/Engine/Plugins/Experimental/NNX/Source/NNXHlslShaders/Private/NNIHlslShadersElementWiseVariadicCS.cpp
@@ -0,0 +1,48 @@
+// Copyright Epic Games, Inc. All Rights Reserved.
+
+#include "NNIHlslShadersElementWiseVariadicCS.h"
+#include "NNXCore.h"
+
+namespace UE::NNIHlslShaders::Internal
+{
+	void TElementWiseVariadicCS::ModifyCompilationEnvironment(const FGlobalShaderPermutationParameters& InParameters, FShaderCompilerEnvironment& OutEnvironment)
+	{
+		FGlobalShader::ModifyCompilationEnvironment(InParameters, OutEnvironment);
+
+		OutEnvironment.SetDefine(TEXT("THREADGROUP_SIZE_X"), FElementWiseVariadicConstants::NUM_GROUP_THREADS);
+
+		FPermutationDomain PermutationVector(InParameters.PermutationId);
+
+		const FString OpFunc = GetOpFunc(PermutationVector.Get<FOperatorType>());
+
+		OutEnvironment.SetDefine(TEXT("ELEMENTWISE_OP(X,Y)"), *OpFunc);
+	}
+
+	const FString TElementWiseVariadicCS::GetOpFunc(EMLElementWiseVariadicOperatorType OpType)
+	{
+		FString OpTable[(int32)EMLElementWiseVariadicOperatorType::MAX];
+
+		for (int32 Idx = 0; Idx < (int32)EMLElementWiseVariadicOperatorType::MAX; ++Idx)
+		{
+			OpTable[Idx] = FString("");
+		}
+
+#define OP(OpName, OpFunc) OpTable[(int32) EMLElementWiseVariadicOperatorType::OpName] = OpFunc
+		OP(Max, TEXT("max(X,Y)"));
+		OP(Min, TEXT("min(X,Y)"));
+		OP(Mean, TEXT("((X)+(Y))"));
+		OP(Sum, TEXT("((X)+(Y))"));
+#undef OP
+
+		FString OpFunc = OpTable[(int32)OpType];
+
+		if (OpFunc == "")
+		{
+			UE_LOG(LogNNX, Warning, TEXT("Undefined ElementWise Variadic operator name for operator:%d"), OpType);
+		}
+
+		return OpFunc;
+	}
+
+	IMPLEMENT_GLOBAL_SHADER(TElementWiseVariadicCS, "/NNX/ElementWiseVariadicOp.usf", "ElementWiseVariadicOp", SF_Compute);
+} // UE::NNIHlslShaders::Internal
\ No newline at end of file
diff --git a/Engine/Plugins/Experimental/NNX/Source/NNXHlslShaders/Private/NNXElementWiseBinaryCS.cpp b/Engine/Plugins/Experimental/NNX/Source/NNXHlslShaders/Private/NNXElementWiseBinaryCS.cpp
deleted file mode 100644
diff --git a/Engine/Plugins/Experimental/NNX/Source/NNXHlslShaders/Private/NNXElementWiseCS.cpp b/Engine/Plugins/Experimental/NNX/Source/NNXHlslShaders/Private/NNXElementWiseCS.cpp
deleted file mode 100644
diff --git a/Engine/Plugins/Experimental/NNX/Source/NNXHlslShaders/Private/NNXElementWiseVariadicCS.cpp b/Engine/Plugins/Experimental/NNX/Source/NNXHlslShaders/Private/NNXElementWiseVariadicCS.cpp
deleted file mode 100644
diff --git a/Engine/Plugins/Experimental/NNX/Source/NNXHlslShaders/Public/NNXElementWiseBinaryCS.h b/Engine/Plugins/Experimental/NNX/Source/NNXHlslShaders/Public/NNXElementWiseBinaryCS.h
deleted file mode 100644
diff --git a/Engine/Plugins/Experimental/NNX/Source/NNXHlslShaders/Public/NNXElementWiseCS.h b/Engine/Plugins/Experimental/NNX/Source/NNXHlslShaders/Public/NNXElementWiseCS.h
deleted file mode 100644
diff --git a/Engine/Plugins/Experimental/NNX/Source/NNXHlslShaders/Public/NNXElementWiseVariadicCS.h b/Engine/Plugins/Experimental/NNX/Source/NNXHlslShaders/Public/NNXElementWiseVariadicCS.h
deleted file mode 100644
diff --git a/Engine/Plugins/Experimental/NNX/Source/NNXRuntimeRDG/Private/Hlsl/NNIRuntimeRDGElementWiseBinary.cpp b/Engine/Plugins/Experimental/NNX/Source/NNXRuntimeRDG/Private/Hlsl/NNIRuntimeRDGElementWiseBinary.cpp
new file mode 100644
--- /dev/null
+++ b/Engine/Plugins/Experimental/NNX/Source/NNXRuntimeRDG/Private/Hlsl/NNIRuntimeRDGElementWiseBinary.cpp
@@ -0,0 +1,113 @@
+// Copyright Epic Games, Inc. All Rights Reserved.
+
+#include "NNIRuntimeRDGElementWiseBinary.h"
+#include "NNIHlslShadersElementWiseBinaryCS.h"
+#include "NNXRuntimeHLSLHelper.h"
+
+namespace UE::NNIRuntimeRDG::Private::Hlsl
+{
+	DECLARE_GPU_STAT_NAMED(FNNIOperatorElementWiseBinary, TEXT("NNI.Operator.Hlsl.ElementWise.Binary"));
+
+	using TElementWiseBinaryCS = typename UE::NNIHlslShaders::Internal::TElementWiseBinaryCS;
+	using FElementWiseBinaryConstants = UE::NNIHlslShaders::Internal::FElementWiseBinaryConstants;
+
+	/**
+	 * Binary element-wise operator implementation
+	 */
+	template<EMLElementWiseBinaryOperatorType OpType>
+	class TOperatorElementWiseBinary : public NNX::FMLOperatorHlsl
+	{
+	public:
+
+		TOperatorElementWiseBinary() {}
+		virtual ~TOperatorElementWiseBinary() = default;
+
+	private:
+
+		NNX::FMLTensorDesc LHSInput;
+		NNX::FMLTensorDesc RHSInput;
+		NNX::FMLTensorDesc Output;
+
+	public:
+
+		virtual bool Initialize(TArrayView<const NNX::FMLTensorDesc> InputTensors, TArrayView<const NNX::FMLTensorDesc> OutputTensors, const FMLAttributeMap& Attributes) override
+		{
+			check(InputTensors.Num() == 2);
+			check(OutputTensors.Num() == 1);
+		
+			LHSInput = InputTensors[0];
+			RHSInput = InputTensors[1];
+			Output = OutputTensors[0];
+
+			return true;
+		}
+
+		virtual void Dispatch(FRDGBuilder& GraphBuilder, TArrayView<const NNX::FMLTensorBinding> InInputBindings, TArrayView<const NNX::FMLTensorBinding> OutOutputBindings) override
+		{
+			// HACK: This only works for single layer networks
+			FRDGBufferSRVRef LHSInputSRV = GraphBuilder.CreateSRV(FRDGBufferSRVDesc(InInputBindings[0].Buffer, PF_R32_FLOAT));
+			FRDGBufferSRVRef RHSInputSRV = GraphBuilder.CreateSRV(FRDGBufferSRVDesc(InInputBindings[1].Buffer, PF_R32_FLOAT));
+			FRDGBufferUAVRef OutputUAV = GraphBuilder.CreateUAV(FRDGBufferUAVDesc(OutOutputBindings[0].Buffer, PF_R32_FLOAT));
+
+			FIntVector ThreadGroupCount = NNX::ComputeElementWiseThreadGroups(Output.Num(), FElementWiseBinaryConstants::NUM_GROUP_THREADS);
+
+			// Set parameters
+			TElementWiseBinaryCS::FParameters* Params = GraphBuilder.AllocParameters<TElementWiseBinaryCS::FParameters>();
+			Params->LHSInput = LHSInputSRV;
+			Params->RHSInput = RHSInputSRV;
+			Params->Output = OutputUAV;
+			FillTensorStrideForBroadcastShaderParameters(LHSInput, Output.Shape.Num(), Params->TensorInfo, 0);
+			FillTensorStrideForBroadcastShaderParameters(RHSInput, Output.Shape.Num(), Params->TensorInfo, 1);
+			FillTensorStrideShaderParameters(Output, Params->TensorInfo, 2);
+			Params->Num = Output.Num();
+			Params->ThreadCountX = ThreadGroupCount.X * FElementWiseBinaryConstants::NUM_GROUP_THREADS;
+
+			TElementWiseBinaryCS::FPermutationDomain PermutationVector;
+
+			PermutationVector.Set<TElementWiseBinaryCS::FOperatorType>(OpType);
+			PermutationVector.Set<TElementWiseBinaryCS::FBinaryNumDimensions>(Output.Shape.Num());
+
+			TShaderMapRef<TElementWiseBinaryCS> ComputeShader(GetGlobalShaderMap(GMaxRHIFeatureLevel), PermutationVector);
+		
+			RDG_EVENT_SCOPE(GraphBuilder, "NNI.Operator.Hlsl.ElementWise.Binary");
+			RDG_GPU_STAT_SCOPE(GraphBuilder, FNNIOperatorElementWiseBinary);
+
+			FComputeShaderUtils::AddPass(
+				GraphBuilder,
+				RDG_EVENT_NAME("NNI.Operator.Hlsl.ElementWise.Binary.Dispatch"),
+				ERDGPassFlags::Compute | ERDGPassFlags::NeverCull,
+				ComputeShader,
+				Params,
+				ThreadGroupCount);
+		}
+	};
+
+	template<EMLElementWiseBinaryOperatorType OpType>
+	NNX::FMLOperatorHlsl* CreateElementWiseBinaryOperator()
+	{
+		return new TOperatorElementWiseBinary<OpType>();
+	}
+
+	bool RegisterElementWiseBinaryOperators(NNX::FMLOperatorRegistryHlsl& Registry)
+	{
+#define OP(Name) Registry.OpAdd(TEXT(#Name), CreateElementWiseBinaryOperator<EMLElementWiseBinaryOperatorType::Name>)
+		OP(Add);
+		//OP(And);
+		OP(Div);
+		//OP(Equal);
+		//OP(Greater);
+		//OP(GreaterOrEqual);
+		//OP(Less);
+		//OP(LessOrEqual);
+		OP(Mod);
+		OP(Mul);
+		//OP(Or);
+		OP(Prelu);
+		OP(Pow);
+		OP(Sub);
+		//OP(Or);
+#undef OP
+
+		return true;
+	}
+} // UE::NNIRuntimeRDG::Private::Hlsl
diff --git a/Engine/Plugins/Experimental/NNX/Source/NNXRuntimeRDG/Private/Hlsl/NNIRuntimeRDGElementWiseBinary.h b/Engine/Plugins/Experimental/NNX/Source/NNXRuntimeRDG/Private/Hlsl/NNIRuntimeRDGElementWiseBinary.h
new file mode 100644
--- /dev/null
+++ b/Engine/Plugins/Experimental/NNX/Source/NNXRuntimeRDG/Private/Hlsl/NNIRuntimeRDGElementWiseBinary.h
@@ -0,0 +1,11 @@
+// Copyright Epic Games, Inc. All Rights Reserved.
+
+#pragma once
+
+#include "../NNXRuntimeRDG.h"
+#include "../NNXRuntimeHLSLOp.h"
+
+namespace UE::NNIRuntimeRDG::Private::Hlsl
+{
+	bool RegisterElementWiseBinaryOperators(NNX::FMLOperatorRegistryHlsl& Registry);
+} // UE::NNIRuntimeRDG::Private::Hlsl
diff --git a/Engine/Plugins/Experimental/NNX/Source/NNXRuntimeRDG/Private/Hlsl/NNIRuntimeRDGElementWiseUnary.cpp b/Engine/Plugins/Experimental/NNX/Source/NNXRuntimeRDG/Private/Hlsl/NNIRuntimeRDGElementWiseUnary.cpp
new file mode 100644
--- /dev/null
+++ b/Engine/Plugins/Experimental/NNX/Source/NNXRuntimeRDG/Private/Hlsl/NNIRuntimeRDGElementWiseUnary.cpp
@@ -0,0 +1,158 @@
+// Copyright Epic Games, Inc. All Rights Reserved.
+
+#include "NNIRuntimeRDGElementWiseUnary.h"
+#include "NNIHlslShadersElementWiseUnaryCS.h"
+#include "NNXRuntimeHLSLHelper.h"
+
+namespace UE::NNIRuntimeRDG::Private::Hlsl
+{
+	DECLARE_GPU_STAT_NAMED(FNNIOperatorElementWiseUnary, TEXT("NNI.Operator.Hlsl.ElementWise.Unary"));
+
+	using TElementWiseUnaryCS = typename UE::NNIHlslShaders::Internal::TElementWiseUnaryCS;
+	using FElementWiseUnaryConstants = UE::NNIHlslShaders::Internal::FElementWiseUnaryConstants;
+
+	/**
+	 * Unary element-wise operator implementation
+	 */
+	template<EMLElementWiseUnaryOperatorType OpType>
+	class TOperatorElementWiseUnary : public NNX::FMLOperatorHlsl
+	{
+	public:
+
+		TOperatorElementWiseUnary() : Alpha(0.0f), Beta(0.0f), Gamma(0.0f) {}
+		virtual ~TOperatorElementWiseUnary() = default;
+
+	private:
+
+		float Alpha;
+		float Beta;
+		float Gamma;
+		NNX::FMLTensorDesc Input;
+		NNX::FMLTensorDesc Output;
+
+	public:
+
+		virtual bool Initialize(TArrayView<const NNX::FMLTensorDesc> InputTensors, TArrayView<const NNX::FMLTensorDesc> OutputTensors, const FMLAttributeMap& Attributes) override
+		{
+			check(InputTensors.Num() == 1);
+			check(OutputTensors.Num() == 1);
+
+			Input = InputTensors[0];
+			Output = OutputTensors[0];
+
+			Alpha = Attributes.GetOptionalFloat(TEXT("alpha"), Alpha);
+			Beta = Attributes.GetOptionalFloat(TEXT("beta"), Beta);
+			Gamma = Attributes.GetOptionalFloat(TEXT("gamma"), Gamma);
+
+			return true;
+		}
+
+		virtual void Dispatch(FRDGBuilder& GraphBuilder, TArrayView<const NNX::FMLTensorBinding> InInputBindings, TArrayView<const NNX::FMLTensorBinding> OutOutputBindings) override
+		{
+			// HACK: This only works for single layer networks
+			FRDGBufferSRVRef InputSRV = GraphBuilder.CreateSRV(FRDGBufferSRVDesc(InInputBindings[0].Buffer, PF_R32_FLOAT));
+			FRDGBufferUAVRef OutputUAV = GraphBuilder.CreateUAV(FRDGBufferUAVDesc(OutOutputBindings[0].Buffer, PF_R32_FLOAT));
+		
+			FIntVector ThreadGroupCount = NNX::ComputeElementWiseThreadGroups(Output.Num(), FElementWiseUnaryConstants::NUM_GROUP_THREADS);
+
+			// Set parameters
+			TElementWiseUnaryCS::FParameters* Params = GraphBuilder.AllocParameters<TElementWiseUnaryCS::FParameters>();
+			Params->Input = InputSRV;
+			Params->Output = OutputUAV;
+			Params->Alpha = Alpha;
+			Params->Beta = Beta;
+			Params->Gamma = Gamma;
+			Params->Num = Output.Num();
+			Params->ThreadCountX = ThreadGroupCount.X * FElementWiseUnaryConstants::NUM_GROUP_THREADS;
+
+			TElementWiseUnaryCS::FPermutationDomain PermutationVector;
+
+			PermutationVector.Set<TElementWiseUnaryCS::FOperatorType>(OpType);
+
+			TShaderMapRef<TElementWiseUnaryCS> ComputeShader(GetGlobalShaderMap(GMaxRHIFeatureLevel), PermutationVector);
+
+			RDG_EVENT_SCOPE(GraphBuilder, "NNI.Operator.Hlsl.ElementWise.Unary");
+			RDG_GPU_STAT_SCOPE(GraphBuilder, FNNIOperatorElementWiseUnary);
+		
+			FComputeShaderUtils::AddPass(
+				GraphBuilder,
+				RDG_EVENT_NAME("NNI.Operator.Hlsl.ElementWise.Unary.Dispatch"),
+				ERDGPassFlags::Compute | ERDGPassFlags::NeverCull,
+				ComputeShader,
+				Params,
+				ThreadGroupCount);
+		}
+	};
+
+	template<> TOperatorElementWiseUnary<EMLElementWiseUnaryOperatorType::Selu>::TOperatorElementWiseUnary()
+		: Alpha(1.67326319217681884765625f), Beta(0.0f), Gamma(1.05070102214813232421875f)
+	{
+	}
+
+	template<> TOperatorElementWiseUnary<EMLElementWiseUnaryOperatorType::Elu>::TOperatorElementWiseUnary()
+		: Alpha(1.0f), Beta(0.0f), Gamma(0.0f) 
+	{
+	}
+
+	template<> TOperatorElementWiseUnary<EMLElementWiseUnaryOperatorType::HardSigmoid>::TOperatorElementWiseUnary()
+		: Alpha(0.2f), Beta(0.5f), Gamma(0.0f)
+	{
+	}
+
+	template<> TOperatorElementWiseUnary<EMLElementWiseUnaryOperatorType::LeakyRelu>::TOperatorElementWiseUnary()
+		: Alpha(0.01f), Beta(0.0f), Gamma(0.0f)
+	{
+	}
+
+	template<EMLElementWiseUnaryOperatorType OpType>
+	NNX::FMLOperatorHlsl* CreateElementWiseUnaryOperator()
+	{
+		return new TOperatorElementWiseUnary<OpType>();
+	}
+
+	bool RegisterElementWiseUnaryOperators(NNX::FMLOperatorRegistryHlsl& Registry)
+	{
+#define OP(Name) Registry.OpAdd(TEXT(#Name), CreateElementWiseUnaryOperator<EMLElementWiseUnaryOperatorType::Name>)
+		OP(Abs);
+		OP(Acos);
+		OP(Acosh);
+		OP(Asin);
+		OP(Asinh);
+		OP(Atan);
+		OP(Atanh);
+		//OP(BitShift);
+		//OP(Cast);
+		OP(Ceil);
+		//OP(Clip);
+		OP(Cos);
+		OP(Cosh);
+		OP(Elu);
+		OP(Erf);
+		OP(Exp);
+		OP(Floor);
+		OP(IsInf);
+		OP(IsNan);
+		OP(HardSigmoid);
+		OP(HardSwish);
+		OP(LeakyRelu);
+		OP(Log);
+		OP(Neg);
+		//OP(Not);
+		OP(Reciprocal);
+		OP(Relu);
+		OP(Round);
+		OP(Selu);
+		OP(Sigmoid);
+		OP(Sign);
+		OP(Sin);
+		OP(Sinh);
+		OP(Softplus);
+		OP(Softsign);
+		OP(Sqrt);
+		OP(Tan);
+		OP(Tanh);
+#undef OP
+
+		return true;
+	}
+} // UE::NNIRuntimeRDG::Private::Hlsl
diff --git a/Engine/Plugins/Experimental/NNX/Source/NNXRuntimeRDG/Private/Hlsl/NNIRuntimeRDGElementWiseUnary.h b/Engine/Plugins/Experimental/NNX/Source/NNXRuntimeRDG/Private/Hlsl/NNIRuntimeRDGElementWiseUnary.h
new file mode 100644
--- /dev/null
+++ b/Engine/Plugins/Experimental/NNX/Source/NNXRuntimeRDG/Private/Hlsl/NNIRuntimeRDGElementWiseUnary.h
@@ -0,0 +1,11 @@
+// Copyright Epic Games, Inc. All Rights Reserved.
+
+#pragma once
+
+#include "../NNXRuntimeRDG.h"
+#include "../NNXRuntimeHLSLOp.h"
+
+namespace UE::NNIRuntimeRDG::Private::Hlsl
+{
+	bool RegisterElementWiseUnaryOperators(NNX::FMLOperatorRegistryHlsl& Registry);
+} // UE::NNIRuntimeRDG::Private::Hlsl
diff --git a/Engine/Plugins/Experimental/NNX/Source/NNXRuntimeRDG/Private/Hlsl/NNIRuntimeRDGElementWiseVariadic.cpp b/Engine/Plugins/Experimental/NNX/Source/NNXRuntimeRDG/Private/Hlsl/NNIRuntimeRDGElementWiseVariadic.cpp
new file mode 100644
--- /dev/null
+++ b/Engine/Plugins/Experimental/NNX/Source/NNXRuntimeRDG/Private/Hlsl/NNIRuntimeRDGElementWiseVariadic.cpp
@@ -0,0 +1,170 @@
+// Copyright Epic Games, Inc. All Rights Reserved.
+
+#include "NNIRuntimeRDGElementWiseVariadic.h"
+#include "NNIHlslShadersElementWiseVariadicCS.h"
+#include "NNXRuntimeHLSLHelper.h"
+
+namespace UE::NNIRuntimeRDG::Private::Hlsl
+{
+	DECLARE_GPU_STAT_NAMED(FNNIOperatorElementWiseVariadic, TEXT("NNI.Operator.Hlsl.ElementWise.Variadic"));
+
+	using TElementWiseVariadicCS = typename UE::NNIHlslShaders::Internal::TElementWiseVariadicCS;
+	using FElementWiseVariadicConstants = UE::NNIHlslShaders::Internal::FElementWiseVariadicConstants;
+
+	void AddOneVariadicOpPass(FRDGBuilder& GraphBuilder, 
+		TArrayView<const NNX::FMLTensorBinding> InputBindings,
+		TArrayView<const NNX::FMLTensorDesc> InputDesc,
+		const NNX::FMLTensorBinding& OutputBinding,
+		const NNX::FMLTensorDesc& OutputDesc,
+		bool OutputAsInput,
+		EMLElementWiseVariadicOperatorType OpType,
+		float Scale)
+	{
+		static_assert(FElementWiseVariadicConstants::MAX_NUM_INPUT == 4, "This algorithm need to be adapted to math the shader.");
+		check(InputBindings.Num() == InputDesc.Num());
+		check(InputBindings.Num() <= FElementWiseVariadicConstants::MAX_NUM_INPUT);
+		check(InputBindings.Num() > 0);
+
+		// SRVs & UAV creations
+		FRDGBufferUAVRef OutputUAV = GraphBuilder.CreateUAV(FRDGBufferUAVDesc(OutputBinding.Buffer, PF_R32_FLOAT));
+		FRDGBufferSRVRef InputsSRV[FElementWiseVariadicConstants::MAX_NUM_INPUT] = { nullptr, nullptr, nullptr, nullptr };
+
+		for (int32 i = 0; i < InputBindings.Num(); ++i)
+		{
+			InputsSRV[i] = GraphBuilder.CreateSRV(FRDGBufferSRVDesc(InputBindings[i].Buffer, PF_R32_FLOAT));
+		}
+
+		// Set parameters
+		FIntVector ThreadGroupCount = NNX::ComputeElementWiseThreadGroups(OutputDesc.Num(), FElementWiseVariadicConstants::NUM_GROUP_THREADS);
+		TElementWiseVariadicCS::FParameters* Params = GraphBuilder.AllocParameters<TElementWiseVariadicCS::FParameters>();
+
+		Params->Input0 = InputsSRV[0];
+		Params->Input1 = InputsSRV[1];
+		Params->Input2 = InputsSRV[2];
+		Params->Input3 = InputsSRV[3];
+		Params->Output = OutputUAV;
+		FillTensorStrideForBroadcastShaderParameters(InputDesc[0], OutputDesc.Shape.Num(), Params->InputTensorInfo, 0);
+		if (InputBindings.Num() >= 2)
+		{
+			FillTensorStrideForBroadcastShaderParameters(InputDesc[1], OutputDesc.Shape.Num(), Params->InputTensorInfo, 1);
+		}
+		if (InputBindings.Num() >= 3)
+		{
+			FillTensorStrideForBroadcastShaderParameters(InputDesc[2], OutputDesc.Shape.Num(), Params->InputTensorInfo, 2);
+		}
+		if (InputBindings.Num() >= 4)
+		{
+			FillTensorStrideForBroadcastShaderParameters(InputDesc[3], OutputDesc.Shape.Num(), Params->InputTensorInfo, 3);
+		}
+		FillTensorStrideShaderParameters(OutputDesc, Params->OutputTensorInfo, 0);
+		Params->Num = OutputDesc.Num();
+		Params->ThreadCountX = ThreadGroupCount.X * FElementWiseVariadicConstants::NUM_GROUP_THREADS;
+		Params->Scale = Scale;
+
+		// Shader variation
+		TElementWiseVariadicCS::FPermutationDomain PermutationVector;
+
+		PermutationVector.Set<TElementWiseVariadicCS::FOperatorType>(OpType);
+		PermutationVector.Set<TElementWiseVariadicCS::FApplyScale>(Scale != 1.0f);
+		PermutationVector.Set<TElementWiseVariadicCS::FOutputAsInput>(OutputAsInput);
+		PermutationVector.Set<TElementWiseVariadicCS::FNumInput>(InputBindings.Num());
+		PermutationVector.Set<TElementWiseVariadicCS::FVariadicNumDimensions>(OutputDesc.Shape.Num());
+
+		// Add the pass to RDG
+		TShaderMapRef<TElementWiseVariadicCS> ComputeShader(GetGlobalShaderMap(GMaxRHIFeatureLevel), PermutationVector);
+		FComputeShaderUtils::AddPass(
+			GraphBuilder,
+			RDG_EVENT_NAME("NNI.Operator.Hlsl.ElementWise.Variadic.Dispatch"),
+			ERDGPassFlags::Compute | ERDGPassFlags::NeverCull,
+			ComputeShader,
+			Params,
+			ThreadGroupCount);
+	}
+
+	/**
+	 * Variadic Element-wise ML operator implementation
+	 */
+	template<EMLElementWiseVariadicOperatorType OpType>
+	class TOperatorElementWiseVariadic : public NNX::FMLOperatorHlsl
+	{
+	public:
+
+		TOperatorElementWiseVariadic() {}
+		virtual ~TOperatorElementWiseVariadic() = default;
+
+	private:
+
+		TArray<NNX::FMLTensorDesc> InputDescs;
+		NNX::FMLTensorDesc OutputDesc;
+
+	public:
+
+		virtual bool Initialize(TArrayView<const NNX::FMLTensorDesc> InputTensors, TArrayView<const NNX::FMLTensorDesc> OutputTensors, const FMLAttributeMap& Attributes) override
+		{
+			check(InputTensors.Num() > 0);
+			check(OutputTensors.Num() == 1);
+		
+			InputDescs.Append(InputTensors);
+			OutputDesc = OutputTensors[0];
+
+			return true;
+		}
+
+		virtual void Dispatch(FRDGBuilder& GraphBuilder, TArrayView<const NNX::FMLTensorBinding> InInputBindings, TArrayView<const NNX::FMLTensorBinding> OutOutputBindings) override
+		{
+			check(OutOutputBindings.Num() == 1);
+			check(InInputBindings.Num() == InputDescs.Num());
+
+			RDG_EVENT_SCOPE(GraphBuilder, "NNI.Operator.Hlsl.ElementWise.Variadic");
+			RDG_GPU_STAT_SCOPE(GraphBuilder, FNNIOperatorElementWiseVariadic);
+
+			NNX::FMLTensorBinding PassInputBindings[FElementWiseVariadicConstants::MAX_NUM_INPUT];
+			NNX::FMLTensorDesc PassInputDescs[FElementWiseVariadicConstants::MAX_NUM_INPUT];
+			for (int32 InputOffset = 0; InputOffset < InInputBindings.Num(); InputOffset += FElementWiseVariadicConstants::MAX_NUM_INPUT)
+			{
+				uint32 NumInputLeftToHandle = InInputBindings.Num() - InputOffset;
+				uint32 NumInputForPass = FMath::Min(FElementWiseVariadicConstants::MAX_NUM_INPUT, NumInputLeftToHandle);
+				bool bIsFirstPass = (InputOffset == 0);
+				bool bIsLastPass = (NumInputLeftToHandle <= FElementWiseVariadicConstants::MAX_NUM_INPUT);
+				float Scale = 1.0f;
+
+				for (uint32 i = 0; i < NumInputForPass; ++i)
+				{
+					PassInputBindings[i] = InInputBindings[InputOffset + i];
+					PassInputDescs[i] = InputDescs[InputOffset + i];
+				}
+
+				if (OpType == EMLElementWiseVariadicOperatorType::Mean && bIsLastPass)
+				{
+					Scale = 1.0f / InInputBindings.Num();
+				}
+			
+				AddOneVariadicOpPass(GraphBuilder,
+					MakeArrayView(PassInputBindings, NumInputForPass),
+					MakeArrayView(PassInputDescs, NumInputForPass),
+					OutOutputBindings[0],
+					OutputDesc,
+					!bIsFirstPass,
+					OpType, Scale);
+			}
+		}
+	};
+
+	template<EMLElementWiseVariadicOperatorType OpType>
+	NNX::FMLOperatorHlsl* CreateElementWiseVariadicOperator()
+	{
+		return new TOperatorElementWiseVariadic<OpType>();
+	}
+
+	bool RegisterElementWiseVariadicOperators(NNX::FMLOperatorRegistryHlsl& Registry)
+	{
+#define OP(Name) Registry.OpAdd(TEXT(#Name), CreateElementWiseVariadicOperator<EMLElementWiseVariadicOperatorType::Name>)
+		OP(Max);
+		OP(Min);
+		OP(Mean);
+		OP(Sum);
+#undef OP
+
+		return true;
+	}
+} // UE::NNIRuntimeRDG::Private::Hlsl
diff --git a/Engine/Plugins/Experimental/NNX/Source/NNXRuntimeRDG/Private/Hlsl/NNIRuntimeRDGElementWiseVariadic.h b/Engine/Plugins/Experimental/NNX/Source/NNXRuntimeRDG/Private/Hlsl/NNIRuntimeRDGElementWiseVariadic.h
new file mode 100644
--- /dev/null
+++ b/Engine/Plugins/Experimental/NNX/Source/NNXRuntimeRDG/Private/Hlsl/NNIRuntimeRDGElementWiseVariadic.h
@@ -0,0 +1,11 @@
+// Copyright Epic Games, Inc. All Rights Reserved.
+
+#pragma once
+
+#include "../NNXRuntimeRDG.h"
+#include "../NNXRuntimeHLSLOp.h"
+
+namespace UE::NNIRuntimeRDG::Private::Hlsl
+{
+	bool RegisterElementWiseVariadicOperators(NNX::FMLOperatorRegistryHlsl& Registry);
+} // UE::NNIRuntimeRDG::Private::Hlsl
diff --git a/Engine/Plugins/Experimental/NNX/Source/NNXRuntimeRDG/Private/NNXRuntimeHLSL.cpp b/Engine/Plugins/Experimental/NNX/Source/NNXRuntimeRDG/Private/NNXRuntimeHLSL.cpp
--- a/Engine/Plugins/Experimental/NNX/Source/NNXRuntimeRDG/Private/NNXRuntimeHLSL.cpp
+++ b/Engine/Plugins/Experimental/NNX/Source/NNXRuntimeRDG/Private/NNXRuntimeHLSL.cpp
@@ -1,17 +1,16 @@
 // Copyright Epic Games, Inc. All Rights Reserved.
 
-#include "Hlsl/NNIRuntimeRDGConv.h"
-#include "NNXElementWiseBinaryCS.h"
 #include "NNXInferenceModel.h"
-#include "NNXModelOptimizer.h"
 #include "NNXRuntimeFormat.h"
-#include "NNXRuntimeHLSLConvTransposeOp.h"
-#include "NNXRuntimeHLSLElementWiseBinaryOps.h"
-#include "NNXRuntimeHLSLElementWiseUnaryOps.h"
-#include "NNXRuntimeHLSLElementWiseVariadicOps.h"
+#include "NNXModelOptimizer.h"
+#include "Hlsl/NNIRuntimeRDGConv.h"
+#include "Hlsl/NNIRuntimeRDGElementWiseBinary.h"
+#include "Hlsl/NNIRuntimeRDGElementWiseUnary.h"
+#include "Hlsl/NNIRuntimeRDGElementWiseVariadic.h"
 #include "NNXRuntimeHLSLGemmOp.h"
-#include "NNXRuntimeHLSLHelper.h"
+#include "NNXRuntimeHLSLConvTransposeOp.h"
 #include "NNXRuntimeHLSLMatMulOp.h"
+#include "NNXRuntimeHLSLHelper.h"
 #include "NNXRuntimeRDG.h"
 
 namespace NNX
@@ -161,9 +160,9 @@ public:
 		FMLOperatorRegistryHlsl* registry = FMLOperatorRegistryHlsl::Get();
 		check(registry != nullptr);
 
-		RegisterElementWiseUnaryOperators(*registry);
-		RegisterElementWiseBinaryOperators(*registry);
-		RegisterElementWiseVariadicOperators(*registry);
+		UE::NNIRuntimeRDG::Private::Hlsl::RegisterElementWiseUnaryOperators(*registry);
+		UE::NNIRuntimeRDG::Private::Hlsl::RegisterElementWiseBinaryOperators(*registry);
+		UE::NNIRuntimeRDG::Private::Hlsl::RegisterElementWiseVariadicOperators(*registry);
 		RegisterGemmOperator(*registry);
 		UE::NNIRuntimeRDG::Private::Hlsl::RegisterConvOperator(*registry);
 		RegisterConvTransposeOperator(*registry);
diff --git a/Engine/Plugins/Experimental/NNX/Source/NNXRuntimeRDG/Private/NNXRuntimeHLSLElementWiseBinaryOps.cpp b/Engine/Plugins/Experimental/NNX/Source/NNXRuntimeRDG/Private/NNXRuntimeHLSLElementWiseBinaryOps.cpp
deleted file mode 100644
diff --git a/Engine/Plugins/Experimental/NNX/Source/NNXRuntimeRDG/Private/NNXRuntimeHLSLElementWiseBinaryOps.h b/Engine/Plugins/Experimental/NNX/Source/NNXRuntimeRDG/Private/NNXRuntimeHLSLElementWiseBinaryOps.h
deleted file mode 100644
diff --git a/Engine/Plugins/Experimental/NNX/Source/NNXRuntimeRDG/Private/NNXRuntimeHLSLElementWiseUnaryOps.cpp b/Engine/Plugins/Experimental/NNX/Source/NNXRuntimeRDG/Private/NNXRuntimeHLSLElementWiseUnaryOps.cpp
deleted file mode 100644
diff --git a/Engine/Plugins/Experimental/NNX/Source/NNXRuntimeRDG/Private/NNXRuntimeHLSLElementWiseUnaryOps.h b/Engine/Plugins/Experimental/NNX/Source/NNXRuntimeRDG/Private/NNXRuntimeHLSLElementWiseUnaryOps.h
deleted file mode 100644
diff --git a/Engine/Plugins/Experimental/NNX/Source/NNXRuntimeRDG/Private/NNXRuntimeHLSLElementWiseVariadicOps.cpp b/Engine/Plugins/Experimental/NNX/Source/NNXRuntimeRDG/Private/NNXRuntimeHLSLElementWiseVariadicOps.cpp
deleted file mode 100644
diff --git a/Engine/Plugins/Experimental/NNX/Source/NNXRuntimeRDG/Private/NNXRuntimeHLSLElementWiseVariadicOps.h b/Engine/Plugins/Experimental/NNX/Source/NNXRuntimeRDG/Private/NNXRuntimeHLSLElementWiseVariadicOps.h
deleted file mode 100644
