NNI Add 2 simple models for weights loading and multi layer test

diff --git a/Engine/Plugins/Experimental/NNX/Source/NNXQA/Resources/NNXQATestDesc.json b/Engine/Plugins/Experimental/NNX/Source/NNXQA/Resources/NNXQATestDesc.json
--- a/Engine/Plugins/Experimental/NNX/Source/NNXQA/Resources/NNXQATestDesc.json
+++ b/Engine/Plugins/Experimental/NNX/Source/NNXQA/Resources/NNXQATestDesc.json
@@ -61,6 +61,27 @@
 					]
 				}
 			]
+		},
+		{
+			"category": "SimpleTests",
+			"is_model_test": true,
+			"targets": [
+				{
+					"target": "OneLayerWithWeights",
+					"tags": [ "SmokeTest" ],
+					"runtimes": [
+						{ "name": "NNXRuntimeHlsl",	"skip": true, "comment": "Weights from model not yet implemented."	},
+						{ "name": "NNXRuntimeDml",  "skip": true, "comment": "Weights from model not yet implemented." }
+					]
+				},
+				{
+					"target": "TwoLayersNoWeight",
+					"runtimes": [
+						{ "name": "NNXRuntimeHlsl",	"skip": true, "comment": "Model with multiple layer not yet implemented."	},
+						{ "name": "NNXRuntimeDml",  "skip": true, "comment": "Model with multiple layer not yet implemented." }
+					]
+				}
+			]
 		}
 	],
 
