ConvTransposeCS cleanup.

diff --git a/Engine/Plugins/Experimental/NNX/Source/NNXHlslShaders/Private/NNXConvTransposeCS.cpp b/Engine/Plugins/Experimental/NNX/Source/NNXHlslShaders/Private/NNXConvTransposeCS.cpp
--- a/Engine/Plugins/Experimental/NNX/Source/NNXHlslShaders/Private/NNXConvTransposeCS.cpp
+++ b/Engine/Plugins/Experimental/NNX/Source/NNXHlslShaders/Private/NNXConvTransposeCS.cpp
@@ -2,13 +2,13 @@
 
 #include "NNXConvTransposeCS.h"
 
-void FMLConvTransposeCS::ModifyCompilationEnvironment(const FGlobalShaderPermutationParameters& InParameters, FShaderCompilerEnvironment& OutEnvironment)
+void FConvTransposeCS::ModifyCompilationEnvironment(const FGlobalShaderPermutationParameters& InParameters, FShaderCompilerEnvironment& OutEnvironment)
 {
 	FGlobalShader::ModifyCompilationEnvironment(InParameters, OutEnvironment);
 	OutEnvironment.SetDefine(TEXT("MAX_NUM_STACK_DIMENSIONS"), FConvTransposeConstants::MAX_NUM_DIMENSIONS);
 }
 
-TArray<int32> FMLConvTransposeCS::GetOutputShape(TArray<int32> XShape, TArray<int32> WShape, EConvTransposeAutoPad AutoPad, TArray<int32> Dilations, TArray<int32> Strides, TArray<int32> Pads, TArray<int32> OutputPadding, int32 Group)
+TArray<int32> FConvTransposeCS::GetOutputShape(TArrayView<const uint32> XShape, TArrayView<const uint32> WShape, EConvTransposeAutoPad AutoPad, TArrayView<const int32> Dilations, TArrayView<const int32> Strides, TArrayView<const int32> Pads, TArrayView<const int32> OutputPadding, int32 Group)
 {
 	check(XShape.Num() > 2);
 	check(WShape.Num() == XShape.Num());
@@ -34,7 +34,7 @@ TArray<int32> FMLConvTransposeCS::GetOutputShape(TArray<int32> XShape, TArray<in
 	return Result;
 }
 
-void FMLConvTransposeCS::FillInParameters(EConvTransposeGroupSize GroupSize, TArray<int32> XShape, TArray<int32> WShape, bool HasB, EConvTransposeAutoPad AutoPad, int32 Group, TArray<int32> Dilations, TArray<int32> Strides, TArray<int32> Pads, TArray<int32> OutputPadding, FMLConvTransposeCS::FParameters& Parameters)
+void FConvTransposeCS::FillInParameters(EConvTransposeGroupSize GroupSize, TArrayView<const uint32> XShape, TArrayView<const uint32> WShape, bool HasB, EConvTransposeAutoPad AutoPad, int32 Group, TArrayView<const int32> Dilations, TArrayView<const int32> Strides, TArrayView<const int32> Pads, TArrayView<const int32> OutputPadding, FConvTransposeCS::FParameters& Parameters)
 {
 	check(XShape.Num() > 2);
 	check(WShape.Num() == XShape.Num());
@@ -88,9 +88,9 @@ void FMLConvTransposeCS::FillInParameters(EConvTransposeGroupSize GroupSize, TAr
 
 	Parameters.XBlockSize = XBlockSize;
 
-	Parameters.NumChannelsPerBatch = FMath::Min((int32)((float)GroupThreadStride / (float)WChannelSize), WShape[0]);
+	Parameters.NumChannelsPerBatch = FMath::Min((int32)((float)GroupThreadStride / (float)WChannelSize), (int32)WShape[0]);
 	check(Parameters.NumChannelsPerBatch > 0)
-		Parameters.NumChannelBatches = FMath::DivideAndRoundUp(WShape[0], Parameters.NumChannelsPerBatch);
+		Parameters.NumChannelBatches = FMath::DivideAndRoundUp((int32)WShape[0], Parameters.NumChannelsPerBatch);
 
 	Parameters.WOutputKernelStride = WShape[1] * WChannelSize;
 	Parameters.WChannelBatchSize = Parameters.NumChannelsPerBatch * WShape[1] * WChannelSize;
@@ -100,7 +100,7 @@ void FMLConvTransposeCS::FillInParameters(EConvTransposeGroupSize GroupSize, TAr
 	Parameters.OneDivGroup = 1.0 / (float)Group;
 }
 
-int32 FMLConvTransposeCS::GetNumReadsPerThread(EConvTransposeGroupSize GroupSize, TArray<int32> WShape, TArray<int32> Dilations, TArray<int32> Strides)
+int32 FConvTransposeCS::GetNumReadsPerThread(EConvTransposeGroupSize GroupSize, TArrayView<const uint32> WShape, TArrayView<const int32> Dilations, TArrayView<const int32> Strides)
 {
 	check(WShape.Num() > 2);
 	check(Dilations.Num() == 0 || Dilations.Num() == WShape.Num() - 2);
@@ -133,7 +133,7 @@ int32 FMLConvTransposeCS::GetNumReadsPerThread(EConvTransposeGroupSize GroupSize
 	return -1;
 }
 
-TArray<int32> FMLConvTransposeCS::GetGroupShape(EConvTransposeGroupSize GroupSize, int32 NumDimensions)
+TArray<int32> FConvTransposeCS::GetGroupShape(EConvTransposeGroupSize GroupSize, int32 NumDimensions)
 {
 	check(NumDimensions > 0);
 
@@ -153,7 +153,7 @@ TArray<int32> FMLConvTransposeCS::GetGroupShape(EConvTransposeGroupSize GroupSiz
 	return Result;
 }
 
-FIntVector FMLConvTransposeCS::GetGroupCount(TArray<int32> YShape, TArray<int32> GroupShape)
+FIntVector FConvTransposeCS::GetGroupCount(TArrayView<const int32> YShape, TArrayView<const int32> GroupShape)
 {
 	check(YShape.Num() > 2);
 	check(YShape.Num() == (GroupShape.Num() + 2));
@@ -167,7 +167,7 @@ FIntVector FMLConvTransposeCS::GetGroupCount(TArray<int32> YShape, TArray<int32>
 	return FIntVector(ThreadGroupCountValueX, YShape[1], YShape[0]);
 }
 
-EConvTransposeGroupSize FMLConvTransposeCS::GetMinimalGroupSize(TArray<int32> WShape)
+EConvTransposeGroupSize FConvTransposeCS::GetMinimalGroupSize(TArrayView<const uint32> WShape)
 {
 	int32 NumDimensions = WShape.Num() - 2;
 	int32 WChannelSize = 1;
@@ -187,7 +187,7 @@ EConvTransposeGroupSize FMLConvTransposeCS::GetMinimalGroupSize(TArray<int32> WS
 	return EConvTransposeGroupSize::MAX;
 }
 
-TArray<int32> FMLConvTransposeCS::GetXBlockShape(TArray<int32> GroupShape, TArray<int32> WShape, TArray<int32> Dilations, TArray<int32> Strides)
+TArray<int32> FConvTransposeCS::GetXBlockShape(TArrayView<const int32> GroupShape, TArrayView<const uint32> WShape, TArrayView<const int32> Dilations, TArrayView<const int32> Strides)
 {
 	check(WShape.Num() > 2);
 	check(GroupShape.Num() == WShape.Num() - 2);
@@ -204,7 +204,7 @@ TArray<int32> FMLConvTransposeCS::GetXBlockShape(TArray<int32> GroupShape, TArra
 	return Result;
 }
 
-TArray<int32> FMLConvTransposeCS::GetPadding(TArray<int32> WShape, EConvTransposeAutoPad AutoPad, TArray<int32> Dilations, TArray<int32> Strides, TArray<int32> Pads, TArray<int32> OutputPadding)
+TArray<int32> FConvTransposeCS::GetPadding(TArrayView<const uint32> WShape, EConvTransposeAutoPad AutoPad, TArrayView<const int32> Dilations, TArrayView<const int32> Strides, TArrayView<const int32> Pads, TArrayView<const int32> OutputPadding)
 {
 	check(WShape.Num() > 2);
 	check(Dilations.Num() == 0 || Dilations.Num() == WShape.Num() - 2);
@@ -217,7 +217,7 @@ TArray<int32> FMLConvTransposeCS::GetPadding(TArray<int32> WShape, EConvTranspos
 
 	if (AutoPad == EConvTransposeAutoPad::NOTSET)
 	{
-		return Pads;
+		return TArray<int32>{Pads};
 	}
 	else if (AutoPad == EConvTransposeAutoPad::VALID)
 	{
@@ -242,7 +242,7 @@ TArray<int32> FMLConvTransposeCS::GetPadding(TArray<int32> WShape, EConvTranspos
 	return Result;
 }
 
-int32 FMLConvTransposeCS::GetNumThreadsPerGroup(EConvTransposeGroupSize GroupSize)
+int32 FConvTransposeCS::GetNumThreadsPerGroup(EConvTransposeGroupSize GroupSize)
 {
 	int32 NumThreadsPerGroup = 128;
 	switch (GroupSize)
@@ -264,7 +264,7 @@ int32 FMLConvTransposeCS::GetNumThreadsPerGroup(EConvTransposeGroupSize GroupSiz
 	return NumThreadsPerGroup;
 }
 
-TArray<int32> FMLConvTransposeCS::GetGridShape(TArray<int32> YShape, TArray<int32> GroupShape)
+TArray<int32> FConvTransposeCS::GetGridShape(TArrayView<const int32> YShape, TArrayView<const int32> GroupShape)
 {
 	check(YShape.Num() > 2);
 	check(YShape.Num() == (GroupShape.Num() + 2));
@@ -279,4 +279,4 @@ TArray<int32> FMLConvTransposeCS::GetGridShape(TArray<int32> YShape, TArray<int3
 	return Result;
 }
 
-IMPLEMENT_GLOBAL_SHADER(FMLConvTransposeCS, "/NNX/ConvTransposeOp.usf", "main", SF_Compute);
\ No newline at end of file
+IMPLEMENT_GLOBAL_SHADER(FConvTransposeCS, "/NNX/ConvTransposeOp.usf", "main", SF_Compute);
\ No newline at end of file
diff --git a/Engine/Plugins/Experimental/NNX/Source/NNXHlslShaders/Public/NNXConvTransposeCS.h b/Engine/Plugins/Experimental/NNX/Source/NNXHlslShaders/Public/NNXConvTransposeCS.h
--- a/Engine/Plugins/Experimental/NNX/Source/NNXHlslShaders/Public/NNXConvTransposeCS.h
+++ b/Engine/Plugins/Experimental/NNX/Source/NNXHlslShaders/Public/NNXConvTransposeCS.h
@@ -33,15 +33,16 @@ enum class EConvTransposeAutoPad : uint8
 class FConvTransposeConstants
 {
 public:
+
     static const int32 MAX_NUM_DIMENSIONS{4};
     static const int32 MIN_NUM_READS_PER_THREAD_POW2{1};
     static const int32 MAX_NUM_READS_PER_THREAD_POW2{3};
 };
 
-class NNXHLSLSHADERS_API FMLConvTransposeCS : public FGlobalShader
+class NNXHLSLSHADERS_API FConvTransposeCS : public FGlobalShader
 {
-	DECLARE_GLOBAL_SHADER(FMLConvTransposeCS);
-	SHADER_USE_PARAMETER_STRUCT(FMLConvTransposeCS, FGlobalShader);
+	DECLARE_GLOBAL_SHADER(FConvTransposeCS);
+	SHADER_USE_PARAMETER_STRUCT(FConvTransposeCS, FGlobalShader);
 
 	class FConvTransposeAlgorithm : SHADER_PERMUTATION_ENUM_CLASS("ALGORITHM", EConvTransposeAlgorithm);
 	class FConvTransposeGroupSize : SHADER_PERMUTATION_ENUM_CLASS("GROUP_SIZE", EConvTransposeGroupSize);
@@ -51,52 +52,54 @@ class NNXHLSLSHADERS_API FMLConvTransposeCS : public FGlobalShader
 	using FPermutationDomain = TShaderPermutationDomain<FConvTransposeAlgorithm, FConvTransposeGroupSize, FConvTransposeNumStackDimensions, FConvTransposeNumReadsPerThread, FConvTransposeHasB>;
 
 public:
+
 	BEGIN_SHADER_PARAMETER_STRUCT(FParameters, )
-		SHADER_PARAMETER_RDG_BUFFER_SRV(Buffer<float>, X) \
-		SHADER_PARAMETER_RDG_BUFFER_SRV(Buffer<float>, W) \
-		SHADER_PARAMETER_RDG_BUFFER_UAV(RWBuffer<float>, Y) \
-		SHADER_PARAMETER_RDG_BUFFER_SRV(Buffer<float>, B) \
-		SHADER_PARAMETER_ARRAY(FIntVector4, Dilation_Stride_XBlockStartOffset_DilationXBlockStride, [FConvTransposeConstants::MAX_NUM_DIMENSIONS]) \
-		SHADER_PARAMETER_ARRAY(FIntVector4, GroupStride_GroupShape_GroupThreadStride_StrideXBlockStride, [FConvTransposeConstants::MAX_NUM_DIMENSIONS]) \
-		SHADER_PARAMETER_ARRAY(FIntVector4, YDimension_YMemoryStride_XDimension_XMemoryStride, [FConvTransposeConstants::MAX_NUM_DIMENSIONS]) \
-		SHADER_PARAMETER_ARRAY(FIntVector4, XBlockStartStride_XBlockStride_WDimension_WDimensionDilationXBlockStride, [FConvTransposeConstants::MAX_NUM_DIMENSIONS]) \
-		SHADER_PARAMETER_ARRAY(FVector4f, OneDiv_GroupStride_GroupThreadStride_OneDivStride, [FConvTransposeConstants::MAX_NUM_DIMENSIONS]) \
-		SHADER_PARAMETER(int32, NumWChannels) \
-		SHADER_PARAMETER(int32, NumOutChannelsDivGroup) \
-		SHADER_PARAMETER(int32, YBatchStride) \
-		SHADER_PARAMETER(int32, YOutputKernelStride) \
-		SHADER_PARAMETER(int32, XBatchStride) \
-		SHADER_PARAMETER(int32, XChannelStride) \
-		SHADER_PARAMETER(int32, XBlockSize) \
-		SHADER_PARAMETER(int32, NumChannelBatches) \
-		SHADER_PARAMETER(int32, NumChannelsPerBatch) \
-		SHADER_PARAMETER(int32, WOutputKernelStride) \
-		SHADER_PARAMETER(int32, WChannelBatchSize) \
-		SHADER_PARAMETER(int32, WChannelSize) \
-		SHADER_PARAMETER(float, GroupsDivM) \
-		SHADER_PARAMETER(float, OneDivGroup) \
+		SHADER_PARAMETER_RDG_BUFFER_SRV(Buffer<float>, X)
+		SHADER_PARAMETER_RDG_BUFFER_SRV(Buffer<float>, W)
+		SHADER_PARAMETER_RDG_BUFFER_UAV(RWBuffer<float>, Y)
+		SHADER_PARAMETER_RDG_BUFFER_SRV(Buffer<float>, B)
+		SHADER_PARAMETER_ARRAY(FIntVector4, Dilation_Stride_XBlockStartOffset_DilationXBlockStride, [FConvTransposeConstants::MAX_NUM_DIMENSIONS])
+		SHADER_PARAMETER_ARRAY(FIntVector4, GroupStride_GroupShape_GroupThreadStride_StrideXBlockStride, [FConvTransposeConstants::MAX_NUM_DIMENSIONS])
+		SHADER_PARAMETER_ARRAY(FIntVector4, YDimension_YMemoryStride_XDimension_XMemoryStride, [FConvTransposeConstants::MAX_NUM_DIMENSIONS])
+		SHADER_PARAMETER_ARRAY(FIntVector4, XBlockStartStride_XBlockStride_WDimension_WDimensionDilationXBlockStride, [FConvTransposeConstants::MAX_NUM_DIMENSIONS])
+		SHADER_PARAMETER_ARRAY(FVector4f, OneDiv_GroupStride_GroupThreadStride_OneDivStride, [FConvTransposeConstants::MAX_NUM_DIMENSIONS])
+		SHADER_PARAMETER(int32, NumWChannels)
+		SHADER_PARAMETER(int32, NumOutChannelsDivGroup)
+		SHADER_PARAMETER(int32, YBatchStride)
+		SHADER_PARAMETER(int32, YOutputKernelStride)
+		SHADER_PARAMETER(int32, XBatchStride)
+		SHADER_PARAMETER(int32, XChannelStride)
+		SHADER_PARAMETER(int32, XBlockSize)
+		SHADER_PARAMETER(int32, NumChannelBatches)
+		SHADER_PARAMETER(int32, NumChannelsPerBatch)
+		SHADER_PARAMETER(int32, WOutputKernelStride)
+		SHADER_PARAMETER(int32, WChannelBatchSize)
+		SHADER_PARAMETER(int32, WChannelSize)
+		SHADER_PARAMETER(float, GroupsDivM)
+		SHADER_PARAMETER(float, OneDivGroup)
 	END_SHADER_PARAMETER_STRUCT()
 
 	static void ModifyCompilationEnvironment(const FGlobalShaderPermutationParameters& InParameters, FShaderCompilerEnvironment& OutEnvironment);
 
-	static TArray<int32> GetOutputShape(TArray<int32> XShape, TArray<int32> WShape, EConvTransposeAutoPad AutoPad, TArray<int32> Dilations, TArray<int32> Strides, TArray<int32> Pads, TArray<int32> OutputPadding, int32 Group);
+	static TArray<int32> GetOutputShape(TArrayView<const uint32> XShape, TArrayView<const uint32> WShape, EConvTransposeAutoPad AutoPad, TArrayView<const int32> Dilations, TArrayView<const int32> Strides, TArrayView<const int32> Pads, TArrayView<const int32> OutputPadding, int32 Group);
 
-	static void FillInParameters(EConvTransposeGroupSize GroupSize, TArray<int32> XShape, TArray<int32> WShape, bool HasB, EConvTransposeAutoPad AutoPad, int32 Group, TArray<int32> Dilations, TArray<int32> Strides, TArray<int32> Pads, TArray<int32> OutputPadding, FMLConvTransposeCS::FParameters& Parameters);
+	static void FillInParameters(EConvTransposeGroupSize GroupSize, TArrayView<const uint32> XShape, TArrayView<const uint32> WShape, bool HasB, EConvTransposeAutoPad AutoPad, int32 Group, TArrayView<const int32> Dilations, TArrayView<const int32> Strides, TArrayView<const int32> Pads, TArrayView<const int32> OutputPadding, FConvTransposeCS::FParameters& Parameters);
 
-	static int32 GetNumReadsPerThread(EConvTransposeGroupSize GroupSize, TArray<int32> WShape, TArray<int32> Dilations, TArray<int32> Strides);
+	static int32 GetNumReadsPerThread(EConvTransposeGroupSize GroupSize, TArrayView<const uint32> WShape, TArrayView<const int32> Dilations, TArrayView<const int32> Strides);
 
 	static TArray<int32> GetGroupShape(EConvTransposeGroupSize GroupSize, int32 NumDimensions);
 
-	static FIntVector GetGroupCount(TArray<int32> YShape, TArray<int32> GroupShape);
+	static FIntVector GetGroupCount(TArrayView<const int32> YShape, TArrayView<const int32> GroupShape);
 
-	static EConvTransposeGroupSize GetMinimalGroupSize(TArray<int32> WShape);
+	static EConvTransposeGroupSize GetMinimalGroupSize(TArrayView<const uint32> WShape);
 
 private:
-	static TArray<int32> GetXBlockShape(TArray<int32> GroupShape, TArray<int32> WShape, TArray<int32> Dilations, TArray<int32> Strides);
 
-	static TArray<int32> GetPadding(TArray<int32> WShape, EConvTransposeAutoPad AutoPad, TArray<int32> Dilations, TArray<int32> Strides, TArray<int32> Pads, TArray<int32> OutputPadding);
+	static TArray<int32> GetXBlockShape(TArrayView<const int32> GroupShape, TArrayView<const uint32> WShape, TArrayView<const int32> Dilations, TArrayView<const int32> Strides);
+
+	static TArray<int32> GetPadding(TArrayView<const uint32> WShape, EConvTransposeAutoPad AutoPad, TArrayView<const int32> Dilations, TArrayView<const int32> Strides, TArrayView<const int32> Pads, TArrayView<const int32> OutputPadding);
 
 	static int32 GetNumThreadsPerGroup(EConvTransposeGroupSize GroupSize);
 
-	static TArray<int32> GetGridShape(TArray<int32> YShape, TArray<int32> GroupShape);
+	static TArray<int32> GetGridShape(TArrayView<const int32> YShape, TArrayView<const int32> GroupShape);
 };
\ No newline at end of file
diff --git a/Engine/Plugins/Experimental/NNX/Source/NNXRuntimeRDG/Private/Hlsl/NNIRuntimeRDGConvTranspose.cpp b/Engine/Plugins/Experimental/NNX/Source/NNXRuntimeRDG/Private/Hlsl/NNIRuntimeRDGConvTranspose.cpp
new file mode 100644
--- /dev/null
+++ b/Engine/Plugins/Experimental/NNX/Source/NNXRuntimeRDG/Private/Hlsl/NNIRuntimeRDGConvTranspose.cpp
@@ -0,0 +1,120 @@
+// Copyright Epic Games, Inc. All Rights Reserved.
+
+#include "NNIRuntimeRDGConvTranspose.h"
+#include "NNXConvTransposeCS.h"
+#include "NNXRuntimeHLSLHelper.h"
+
+namespace UE::NNIRuntimeRDG::Private::Hlsl
+{
+	DECLARE_GPU_STAT_NAMED(FNNIOperatorConvTranspose, TEXT("NNI.Operator.Hlsl.ConvTranspose"));
+
+	/**
+	 * ConvTranspose operator implementation
+	 */
+	class FConvTranspose : public NNX::FMLOperatorHlsl
+	{
+	public:
+
+		static NNX::FMLOperatorHlsl* Create()
+		{
+			return new FConvTranspose();
+		}
+
+		virtual ~FConvTranspose() = default;
+
+	private:
+
+		FConvTranspose() {}
+
+		NNX::FMLTensorDesc Input = {};
+		NNX::FMLTensorDesc Weights = {};
+		NNX::FMLTensorDesc Bias = {};
+		NNX::FMLTensorDesc Output = {};
+
+		int NumDimensions = 0;
+		bool HasBias = false;
+
+		// Hard coded parameters, until we accept them from JSON
+		int32 Group = 1;
+		EConvTransposeAutoPad AutoPad = EConvTransposeAutoPad::VALID;
+		TArray<int32> Dilations = {1};
+		TArray<int32> Strides = {1};
+		TArray<int32> Pads = {0, 0};
+		TArray<int32> OutputPadding = {0};
+
+	public:
+
+		virtual bool Initialize(TArrayView<const NNX::FMLTensorDesc> InputTensors, TArrayView<const NNX::FMLTensorDesc> OutputTensors, const FMLAttributeMap& Attributes) override
+		{
+			check(InputTensors.Num() >= 2 && InputTensors.Num() <= 3);
+			check(OutputTensors.Num() == 1);
+
+			Input = InputTensors[0];
+			Weights = InputTensors[1];
+			Output = OutputTensors[0];
+
+			check(Input.Shape.Num() > 2);
+			check(Weights.Shape.Num()  == Input.Shape.Num());
+			check(Output.Shape.Num() == Input.Shape.Num());
+
+			if (InputTensors.Num() == 3) {
+				HasBias = true;
+				Bias = InputTensors[2];
+			}
+			else
+			{
+				HasBias = false;
+			}
+
+			NumDimensions = Input.Shape.Num() - 2;
+
+			return true;
+		}
+
+		virtual void Dispatch(FRDGBuilder& GraphBuilder, TArrayView<const NNX::FMLTensorBinding> InInputBindings, TArrayView<const NNX::FMLTensorBinding> OutOutputBindings) override
+		{
+			constexpr EConvTransposeAlgorithm Algorithm = EConvTransposeAlgorithm::SharedMemory;
+			constexpr EConvTransposeGroupSize GroupSize = EConvTransposeGroupSize::Size256;
+
+			TArray<int32> OutputShape = FConvTransposeCS::GetOutputShape(Input.Shape, Weights.Shape, AutoPad, Dilations, Strides, Pads, OutputPadding, Group);
+
+			// Set parameters
+			FConvTransposeCS::FParameters* Params = GraphBuilder.AllocParameters<FConvTransposeCS::FParameters>();
+			FConvTransposeCS::FillInParameters(GroupSize, Input.Shape, Weights.Shape, HasBias, AutoPad, Group, Dilations,Strides, Pads, OutputPadding, *Params);
+			Params->X = GraphBuilder.CreateSRV(FRDGBufferSRVDesc(InInputBindings[0].Buffer, PF_R32_FLOAT));
+			Params->W = GraphBuilder.CreateSRV(FRDGBufferSRVDesc(InInputBindings[1].Buffer, PF_R32_FLOAT));
+			if (InInputBindings.Num() == 3) {
+				Params->B = GraphBuilder.CreateSRV(FRDGBufferSRVDesc(InInputBindings[2].Buffer, PF_R32_FLOAT));
+			}
+			Params->Y = GraphBuilder.CreateUAV(FRDGBufferUAVDesc(OutOutputBindings[0].Buffer, PF_R32_FLOAT));
+
+			FConvTransposeCS::FPermutationDomain PermutationVector;
+
+			PermutationVector.Set<FConvTransposeCS::FConvTransposeAlgorithm>(Algorithm);
+			PermutationVector.Set<FConvTransposeCS::FConvTransposeGroupSize>(GroupSize);
+			PermutationVector.Set<FConvTransposeCS::FConvTransposeNumStackDimensions>(NumDimensions);
+			PermutationVector.Set<FConvTransposeCS::FConvTransposeNumReadsPerThread>(FConvTransposeCS::GetNumReadsPerThread(GroupSize, Weights.Shape, Dilations, Strides));
+			PermutationVector.Set<FConvTransposeCS::FConvTransposeHasB>(HasBias);
+			TShaderMapRef<FConvTransposeCS> ComputeShader(GetGlobalShaderMap(GMaxRHIFeatureLevel), PermutationVector);
+
+			RDG_EVENT_SCOPE(GraphBuilder, "NNI.Operator.Hlsl.ConvTranspose");
+			RDG_GPU_STAT_SCOPE(GraphBuilder, FNNIOperatorConvTranspose);
+
+			FComputeShaderUtils::AddPass(
+				GraphBuilder,
+				RDG_EVENT_NAME("NNI.Operator.Hlsl.ConvTranspose.Dispatch"),
+				ERDGPassFlags::Compute | ERDGPassFlags::NeverCull,
+				ComputeShader,
+				Params,
+				FConvTransposeCS::GetGroupCount(OutputShape, FConvTransposeCS::GetGroupShape(GroupSize, NumDimensions)));		
+		}
+	};
+
+	bool RegisterConvTransposeOperator(NNX::FMLOperatorRegistryHlsl& Registry)
+	{
+		Registry.OpAdd(TEXT("ConvTranspose"), FConvTranspose::Create);
+
+		return true;
+	}
+
+} // NNX
diff --git a/Engine/Plugins/Experimental/NNX/Source/NNXRuntimeRDG/Private/Hlsl/NNIRuntimeRDGConvTranspose.h b/Engine/Plugins/Experimental/NNX/Source/NNXRuntimeRDG/Private/Hlsl/NNIRuntimeRDGConvTranspose.h
new file mode 100644
--- /dev/null
+++ b/Engine/Plugins/Experimental/NNX/Source/NNXRuntimeRDG/Private/Hlsl/NNIRuntimeRDGConvTranspose.h
@@ -0,0 +1,11 @@
+// Copyright Epic Games, Inc. All Rights Reserved.
+
+#pragma once
+
+#include "NNXRuntimeRDG.h"
+#include "NNXRuntimeHLSLOp.h"
+
+namespace UE::NNIRuntimeRDG::Private::Hlsl
+{
+	bool RegisterConvTransposeOperator(NNX::FMLOperatorRegistryHlsl& Registry);
+} // UE::NNIRuntimeRDG::Private::Hlsl
diff --git a/Engine/Plugins/Experimental/NNX/Source/NNXRuntimeRDG/Private/NNXRuntimeHLSL.cpp b/Engine/Plugins/Experimental/NNX/Source/NNXRuntimeRDG/Private/NNXRuntimeHLSL.cpp
--- a/Engine/Plugins/Experimental/NNX/Source/NNXRuntimeRDG/Private/NNXRuntimeHLSL.cpp
+++ b/Engine/Plugins/Experimental/NNX/Source/NNXRuntimeRDG/Private/NNXRuntimeHLSL.cpp
@@ -1,5 +1,7 @@
 // Copyright Epic Games, Inc. All Rights Reserved.
 
+#include "Hlsl/NNIRuntimeRDGConv.h"
+#include "Hlsl/NNIRuntimeRDGConvTranspose.h"
 #include "NNXInferenceModel.h"
 #include "NNXRuntimeFormat.h"
 #include "NNXModelOptimizer.h"
@@ -8,7 +10,6 @@
 #include "Hlsl/NNIRuntimeRDGElementWiseUnary.h"
 #include "Hlsl/NNIRuntimeRDGElementWiseVariadic.h"
 #include "NNXRuntimeHLSLGemmOp.h"
-#include "NNXRuntimeHLSLConvTransposeOp.h"
 #include "NNXRuntimeHLSLMatMulOp.h"
 #include "NNXRuntimeHLSLHelper.h"
 #include "NNXRuntimeRDG.h"
@@ -165,7 +166,7 @@ public:
 		UE::NNIRuntimeRDG::Private::Hlsl::RegisterElementWiseVariadicOperators(*registry);
 		RegisterGemmOperator(*registry);
 		UE::NNIRuntimeRDG::Private::Hlsl::RegisterConvOperator(*registry);
-		RegisterConvTransposeOperator(*registry);
+		UE::NNIRuntimeRDG::Private::Hlsl::RegisterConvTransposeOperator(*registry);
 		RegisterMatMulOperator(*registry);
 
 		return true;
diff --git a/Engine/Plugins/Experimental/NNX/Source/NNXRuntimeRDG/Private/NNXRuntimeHLSLConvTransposeOp.cpp b/Engine/Plugins/Experimental/NNX/Source/NNXRuntimeRDG/Private/NNXRuntimeHLSLConvTransposeOp.cpp
deleted file mode 100644
diff --git a/Engine/Plugins/Experimental/NNX/Source/NNXRuntimeRDG/Private/NNXRuntimeHLSLConvTransposeOp.h b/Engine/Plugins/Experimental/NNX/Source/NNXRuntimeRDG/Private/NNXRuntimeHLSLConvTransposeOp.h
deleted file mode 100644
