Conclude cleanup/refactor of Conv and ConvTranspose shader and operator.

diff --git a/Engine/Plugins/Experimental/NNX/Source/NNXHlslShaders/Internal/NNIHlslShadersConvCS.h b/Engine/Plugins/Experimental/NNX/Source/NNXHlslShaders/Internal/NNIHlslShadersConvCS.h
new file mode 100644
--- /dev/null
+++ b/Engine/Plugins/Experimental/NNX/Source/NNXHlslShaders/Internal/NNIHlslShadersConvCS.h
@@ -0,0 +1,109 @@
+// Copyright Epic Games, Inc. All Rights Reserved.
+
+#pragma once
+
+#include "GlobalShader.h"
+#include "ShaderParameterUtils.h"
+#include "RenderGraphUtils.h"
+
+namespace UE::NNIHlslShaders::Internal
+{
+	enum class EConvAlgorithm : uint8
+	{
+		SharedMemory = 0,
+		MAX
+	};
+
+	enum class EConvGroupSize : uint8
+	{
+		Size128 = 0,
+		Size256,
+		Size512,
+		MAX
+	};
+
+	enum class EConvAutoPad : uint8
+	{
+		NOTSET = 0,// Use pad values passed in the array
+		SAME_UPPER,// Auto-pad to match input and output shape with potetnial extra padding at the end
+		SAME_LOWER,// Auto-pad to match input and output shape with potetnial extra padding at the beginning
+		VALID,// Set all paddings to zero
+		MAX
+	};
+
+	class FConvConstants
+	{
+	public:
+		static const int32 MAX_NUM_DIMENSIONS{4};
+		static const int32 MIN_NUM_READS_PER_THREAD_POW2{1};
+		static const int32 MAX_NUM_READS_PER_THREAD_POW2{3};
+	};
+
+	class NNXHLSLSHADERS_API FConvCS : public FGlobalShader
+	{
+		DECLARE_GLOBAL_SHADER(FConvCS);
+		SHADER_USE_PARAMETER_STRUCT(FConvCS, FGlobalShader)
+
+		class FConvAlgorithm : SHADER_PERMUTATION_ENUM_CLASS("ALGORITHM", EConvAlgorithm);
+		class FConvGroupSize : SHADER_PERMUTATION_ENUM_CLASS("GROUP_SIZE", EConvGroupSize);
+		class FConvNumDimensions : SHADER_PERMUTATION_RANGE_INT("NUM_DIMENSIONS", 1, FConvConstants::MAX_NUM_DIMENSIONS);
+		class FConvNumReadsPerThread : SHADER_PERMUTATION_RANGE_INT("NUM_READS_PER_THREAD_POW2", FConvConstants::MIN_NUM_READS_PER_THREAD_POW2, FConvConstants::MAX_NUM_READS_PER_THREAD_POW2);
+		class FConvHasB : SHADER_PERMUTATION_BOOL("HAS_B");
+		using FPermutationDomain = TShaderPermutationDomain<FConvAlgorithm, FConvGroupSize, FConvNumDimensions, FConvNumReadsPerThread, FConvHasB>;
+
+	public:
+		BEGIN_SHADER_PARAMETER_STRUCT(FParameters, )
+			SHADER_PARAMETER_RDG_BUFFER_SRV(Buffer<float>, X)
+			SHADER_PARAMETER_RDG_BUFFER_SRV(Buffer<float>, W)
+			SHADER_PARAMETER_RDG_BUFFER_UAV(RWBuffer<float>, Y)
+			SHADER_PARAMETER_RDG_BUFFER_SRV(Buffer<float>, B)
+			SHADER_PARAMETER_ARRAY(FIntVector4, Dilation_Stride_XBlockStartOffset_DilationXBlockStride, [FConvConstants::MAX_NUM_DIMENSIONS])
+			SHADER_PARAMETER_ARRAY(FIntVector4, GroupStride_GroupShape_GroupThreadStride_StrideXBlockStride, [FConvConstants::MAX_NUM_DIMENSIONS])
+			SHADER_PARAMETER_ARRAY(FIntVector4, YDimension_YMemoryStride_XDimension_XMemoryStride, [FConvConstants::MAX_NUM_DIMENSIONS])
+			SHADER_PARAMETER_ARRAY(FIntVector4, XBlockStartStride_XBlockStride_WDimension_WDimensionDilationXBlockStride, [FConvConstants::MAX_NUM_DIMENSIONS])
+			SHADER_PARAMETER_ARRAY(FVector4f, OneDiv_GroupStride_GroupThreadStride_XBlockStride, [FConvConstants::MAX_NUM_DIMENSIONS])
+			SHADER_PARAMETER(int32, NumWChannels)
+			SHADER_PARAMETER(int32, YBatchStride)
+			SHADER_PARAMETER(int32, YOutputKernelStride)
+			SHADER_PARAMETER(int32, XBatchStride)
+			SHADER_PARAMETER(int32, XChannelStride)
+			SHADER_PARAMETER(int32, XBlockSize)
+			SHADER_PARAMETER(int32, NumChannelBatches)
+			SHADER_PARAMETER(int32, NumChannelsPerBatch)
+			SHADER_PARAMETER(int32, WOutputKernelStride)
+			SHADER_PARAMETER(int32, WChannelBatchSize)
+			SHADER_PARAMETER(int32, WChannelSize)
+			SHADER_PARAMETER(float, GroupsDivM)
+		END_SHADER_PARAMETER_STRUCT()
+
+		static void ModifyCompilationEnvironment(const FGlobalShaderPermutationParameters& InParameters, FShaderCompilerEnvironment& OutEnvironment);
+
+	public:
+		static TArray<int32> GetOutputShape(TArrayView<const uint32> XShape, TArrayView<const uint32> WShape, EConvAutoPad AutoPad, TArrayView<const int32> Dilations, TArrayView<const int32> Strides, TArrayView<const int32> Pads);
+
+		static void FillInParameters(EConvGroupSize GroupSize, TArrayView<const uint32> XShape, TArrayView<const uint32> WShape, bool HasB,
+				EConvAutoPad AutoPad, int Group, TArrayView<const int32> Dilations, TArrayView<const int32> Strides, TArrayView<const int32> Pads, FConvCS::FParameters& Parameters);
+
+		static int32 GetNumReadsPerThread(EConvGroupSize GroupSize, TArrayView<const uint32> WShape, TArrayView<const int32> Dilations, TArrayView<const int32> Strides);
+
+		/**
+		* @brief Computes the group shape such that all dimension have roughly equal sizes.
+		*
+		* @param GroupSize The enum indicating the number of threads contained by a single group.
+		* @param NumDimensions The number of dimensions.
+		* @return TArray<int32> An array of size \p NumDimensions containing the number of threads in each dimension to form a volume of a total number of threads indicated by \p GroupSize
+		*/
+		static TArray<int32> GetGroupShape(EConvGroupSize GroupSize, int32 NumDimensions);
+
+		/**
+		* @brief Get the group count vector used to launch the gpu shader thread groups
+		*
+		* @param YShape The shape of the output as computed by GetOutputShape()
+		* @param YShape The shape of a single thread group as computed by GetGroupShape()
+		* @return FIntVector The number of thread groups to instantiate. z corresponds to the batch and y to the output kernel.
+		*/
+		static FIntVector GetGroupCount(TArrayView<const int32> YShape, TArrayView<const int32> GroupShape);
+
+		static EConvGroupSize GetMinimalGroupSize(TArrayView<const int32> WShape);
+	};
+} // UE::NNIHlslShaders::Internal
\ No newline at end of file
diff --git a/Engine/Plugins/Experimental/NNX/Source/NNXHlslShaders/Internal/NNIHlslShadersConvTransposeCS.h b/Engine/Plugins/Experimental/NNX/Source/NNXHlslShaders/Internal/NNIHlslShadersConvTransposeCS.h
new file mode 100644
--- /dev/null
+++ b/Engine/Plugins/Experimental/NNX/Source/NNXHlslShaders/Internal/NNIHlslShadersConvTransposeCS.h
@@ -0,0 +1,108 @@
+// Copyright Epic Games, Inc. All Rights Reserved.
+
+#pragma once
+
+#include "GlobalShader.h"
+#include "ShaderParameterUtils.h"
+#include "RenderGraphUtils.h"
+#include "NNXTypes.h"
+
+namespace UE::NNIHlslShaders::Internal
+{
+	enum class EConvTransposeAlgorithm : uint8
+	{
+		SharedMemory = 0,
+		MAX
+	};
+
+	enum class EConvTransposeGroupSize : uint8
+	{
+		Size128 = 0,
+		Size256,
+		Size512,
+		MAX
+	};
+
+	enum class EConvTransposeAutoPad : uint8
+	{
+		NOTSET = 0,// Use pad values passed in the array
+		SAME_UPPER,// Auto-pad to match input and output shape with potetnial extra padding at the end
+		SAME_LOWER,// Auto-pad to match input and output shape with potetnial extra padding at the beginning
+		VALID,// Set all paddings to zero
+		MAX
+	};
+
+	class FConvTransposeConstants
+	{
+	public:
+
+		static const int32 MAX_NUM_DIMENSIONS{4};
+		static const int32 MIN_NUM_READS_PER_THREAD_POW2{1};
+		static const int32 MAX_NUM_READS_PER_THREAD_POW2{3};
+	};
+
+	class NNXHLSLSHADERS_API FConvTransposeCS : public FGlobalShader
+	{
+		DECLARE_GLOBAL_SHADER(FConvTransposeCS);
+		SHADER_USE_PARAMETER_STRUCT(FConvTransposeCS, FGlobalShader);
+
+		class FConvTransposeAlgorithm : SHADER_PERMUTATION_ENUM_CLASS("ALGORITHM", EConvTransposeAlgorithm);
+		class FConvTransposeGroupSize : SHADER_PERMUTATION_ENUM_CLASS("GROUP_SIZE", EConvTransposeGroupSize);
+		class FConvTransposeNumStackDimensions : SHADER_PERMUTATION_RANGE_INT("NUM_STACK_DIMENSIONS", 1, FConvTransposeConstants::MAX_NUM_DIMENSIONS);
+		class FConvTransposeNumReadsPerThread : SHADER_PERMUTATION_RANGE_INT("NUM_READS_PER_THREAD_POW2", FConvTransposeConstants::MIN_NUM_READS_PER_THREAD_POW2, FConvTransposeConstants::MAX_NUM_READS_PER_THREAD_POW2);
+		class FConvTransposeHasB : SHADER_PERMUTATION_BOOL("HAS_B");
+		using FPermutationDomain = TShaderPermutationDomain<FConvTransposeAlgorithm, FConvTransposeGroupSize, FConvTransposeNumStackDimensions, FConvTransposeNumReadsPerThread, FConvTransposeHasB>;
+
+	public:
+
+		BEGIN_SHADER_PARAMETER_STRUCT(FParameters, )
+			SHADER_PARAMETER_RDG_BUFFER_SRV(Buffer<float>, X)
+			SHADER_PARAMETER_RDG_BUFFER_SRV(Buffer<float>, W)
+			SHADER_PARAMETER_RDG_BUFFER_UAV(RWBuffer<float>, Y)
+			SHADER_PARAMETER_RDG_BUFFER_SRV(Buffer<float>, B)
+			SHADER_PARAMETER_ARRAY(FIntVector4, Dilation_Stride_XBlockStartOffset_DilationXBlockStride, [FConvTransposeConstants::MAX_NUM_DIMENSIONS])
+			SHADER_PARAMETER_ARRAY(FIntVector4, GroupStride_GroupShape_GroupThreadStride_StrideXBlockStride, [FConvTransposeConstants::MAX_NUM_DIMENSIONS])
+			SHADER_PARAMETER_ARRAY(FIntVector4, YDimension_YMemoryStride_XDimension_XMemoryStride, [FConvTransposeConstants::MAX_NUM_DIMENSIONS])
+			SHADER_PARAMETER_ARRAY(FIntVector4, XBlockStartStride_XBlockStride_WDimension_WDimensionDilationXBlockStride, [FConvTransposeConstants::MAX_NUM_DIMENSIONS])
+			SHADER_PARAMETER_ARRAY(FVector4f, OneDiv_GroupStride_GroupThreadStride_OneDivStride, [FConvTransposeConstants::MAX_NUM_DIMENSIONS])
+			SHADER_PARAMETER(int32, NumWChannels)
+			SHADER_PARAMETER(int32, NumOutChannelsDivGroup)
+			SHADER_PARAMETER(int32, YBatchStride)
+			SHADER_PARAMETER(int32, YOutputKernelStride)
+			SHADER_PARAMETER(int32, XBatchStride)
+			SHADER_PARAMETER(int32, XChannelStride)
+			SHADER_PARAMETER(int32, XBlockSize)
+			SHADER_PARAMETER(int32, NumChannelBatches)
+			SHADER_PARAMETER(int32, NumChannelsPerBatch)
+			SHADER_PARAMETER(int32, WOutputKernelStride)
+			SHADER_PARAMETER(int32, WChannelBatchSize)
+			SHADER_PARAMETER(int32, WChannelSize)
+			SHADER_PARAMETER(float, GroupsDivM)
+			SHADER_PARAMETER(float, OneDivGroup)
+		END_SHADER_PARAMETER_STRUCT()
+
+		static void ModifyCompilationEnvironment(const FGlobalShaderPermutationParameters& InParameters, FShaderCompilerEnvironment& OutEnvironment);
+
+		static TArray<int32> GetOutputShape(TArrayView<const uint32> XShape, TArrayView<const uint32> WShape, EConvTransposeAutoPad AutoPad, TArrayView<const int32> Dilations, TArrayView<const int32> Strides, TArrayView<const int32> Pads, TArrayView<const int32> OutputPadding, int32 Group);
+
+		static void FillInParameters(EConvTransposeGroupSize GroupSize, TArrayView<const uint32> XShape, TArrayView<const uint32> WShape, bool HasB, EConvTransposeAutoPad AutoPad, int32 Group, TArrayView<const int32> Dilations, TArrayView<const int32> Strides, TArrayView<const int32> Pads, TArrayView<const int32> OutputPadding, FConvTransposeCS::FParameters& Parameters);
+
+		static int32 GetNumReadsPerThread(EConvTransposeGroupSize GroupSize, TArrayView<const uint32> WShape, TArrayView<const int32> Dilations, TArrayView<const int32> Strides);
+
+		static TArray<int32> GetGroupShape(EConvTransposeGroupSize GroupSize, int32 NumDimensions);
+
+		static FIntVector GetGroupCount(TArrayView<const int32> YShape, TArrayView<const int32> GroupShape);
+
+		static EConvTransposeGroupSize GetMinimalGroupSize(TArrayView<const uint32> WShape);
+
+	private:
+
+		static TArray<int32> GetXBlockShape(TArrayView<const int32> GroupShape, TArrayView<const uint32> WShape, TArrayView<const int32> Dilations, TArrayView<const int32> Strides);
+
+		static TArray<int32> GetPadding(TArrayView<const uint32> WShape, EConvTransposeAutoPad AutoPad, TArrayView<const int32> Dilations, TArrayView<const int32> Strides, TArrayView<const int32> Pads, TArrayView<const int32> OutputPadding);
+
+		static int32 GetNumThreadsPerGroup(EConvTransposeGroupSize GroupSize);
+
+		static TArray<int32> GetGridShape(TArrayView<const int32> YShape, TArrayView<const int32> GroupShape);
+	};
+} // UE::NNIHlslShaders::Internal
\ No newline at end of file
diff --git a/Engine/Plugins/Experimental/NNX/Source/NNXHlslShaders/Private/NNIHlslShadersConvCS.cpp b/Engine/Plugins/Experimental/NNX/Source/NNXHlslShaders/Private/NNIHlslShadersConvCS.cpp
new file mode 100644
--- /dev/null
+++ b/Engine/Plugins/Experimental/NNX/Source/NNXHlslShaders/Private/NNIHlslShadersConvCS.cpp
@@ -0,0 +1,292 @@
+// Copyright Epic Games, Inc. All Rights Reserved.
+
+#include "NNIHlslShadersConvCS.h"
+#include "NNXCore.h"
+
+namespace UE::NNIHlslShaders::Internal
+{
+	namespace
+	{
+
+	TArray<int32> GetXBlockShape(TArrayView<const int32> GroupShape, TArrayView<const uint32> WShape, TArrayView<const int32> Dilations, TArrayView<const int32> Strides)
+	{
+		check(WShape.Num() > 2);
+		check(GroupShape.Num() == WShape.Num() - 2);
+		check(Dilations.Num() == 0 || Dilations.Num() == GroupShape.Num());
+		check(Strides.Num() == 0 || Strides.Num() == GroupShape.Num());
+		
+		TArray<int32> Result;
+		Result.SetNumUninitialized(GroupShape.Num());
+		for (int32 i = 0; i < GroupShape.Num(); i++)
+		{
+			int32 DilatedKernelSize = (i < Dilations.Num() ? Dilations[i] : 1) * (WShape[2 + i] - 1) + 1;
+			Result[i] = DilatedKernelSize + (GroupShape[i] - 1) * (i < Strides.Num() ? Strides[i] : 1); 
+		}
+		return Result;
+	}
+
+	TArray<int32> GetPadding(TArrayView<const uint32> XShape, TArrayView<const uint32> WShape, EConvAutoPad AutoPad, TArrayView<const int32> Dilations, TArrayView<const int32> Strides, TArrayView<const int32> Pads)
+	{
+		check(XShape.Num() > 2);
+		check(WShape.Num() == XShape.Num());
+		check(Dilations.Num() == 0 || Dilations.Num() == WShape.Num() - 2);
+		check(Strides.Num() == 0 || Strides.Num() == WShape.Num() - 2);
+		check(AutoPad != EConvAutoPad::NOTSET || Pads.Num() == 2 * (WShape.Num() - 2));
+
+		int32 NumDimensions = XShape.Num() - 2;
+		TArray<int32> Result;
+		Result.Init(0, 2 * NumDimensions);
+
+		if (AutoPad == EConvAutoPad::NOTSET)
+		{
+			return TArray<int32>{Pads};
+		}
+		else if (AutoPad == EConvAutoPad::VALID)
+		{
+			return Result;
+		}
+
+		for (int32 DimensionIndex = 0; DimensionIndex < NumDimensions; DimensionIndex++)
+		{
+			int32 DilatedKernelSize = (DimensionIndex < Dilations.Num() ? Dilations[DimensionIndex] : 1) * (WShape[DimensionIndex + 2] - 1) + 1;
+			int32 TotalPad = ((int32)(((XShape[DimensionIndex + 2] + (DimensionIndex < Strides.Num() ? Strides[DimensionIndex] : 1) - 1) / (DimensionIndex < Strides.Num() ? Strides[DimensionIndex] : 1)) - 1)) * (DimensionIndex < Strides.Num() ? Strides[DimensionIndex] : 1) + DilatedKernelSize - XShape[DimensionIndex + 2];
+			if (AutoPad == EConvAutoPad::SAME_LOWER)
+			{
+				Result[DimensionIndex] = (TotalPad + 1) / 2;
+			}	
+			else
+			{
+				Result[DimensionIndex] = TotalPad / 2;
+			}
+			Result[NumDimensions + DimensionIndex] = TotalPad - Result[DimensionIndex];
+		}
+
+		return Result;
+	}
+
+	int32 GetNumThreadsPerGroup(EConvGroupSize GroupSize)
+	{
+		int32 NumThreadsPerGroup = 128;
+		switch (GroupSize)
+		{
+			case EConvGroupSize::Size128:
+				NumThreadsPerGroup = 128;
+				break;
+			case EConvGroupSize::Size256:
+				NumThreadsPerGroup = 256;
+				break;
+			case EConvGroupSize::Size512:
+				NumThreadsPerGroup = 512;
+				break;
+			default:
+				NumThreadsPerGroup = 128;
+				break;
+		}
+		check(FMath::Log2((float)NumThreadsPerGroup) == FMath::Floor(FMath::Log2((float)NumThreadsPerGroup)));
+		return NumThreadsPerGroup;
+	}
+
+	TArray<int32> GetGridShape(TArrayView<const int32> YShape, TArrayView<const int32> GroupShape)
+	{
+		check(YShape.Num() > 2);
+		check(YShape.Num() == (GroupShape.Num() + 2));
+
+		TArray<int32> Result;
+		Result.SetNumUninitialized(YShape.Num() - 2);
+		for (int32 i = 2; i < YShape.Num(); i++)
+		{
+			Result[i-2] = FMath::DivideAndRoundUp(YShape[i], GroupShape[i-2]);
+		}
+
+		return Result;
+	}
+
+	} // namespace
+
+	void FConvCS::ModifyCompilationEnvironment(const FGlobalShaderPermutationParameters& InParameters, FShaderCompilerEnvironment& OutEnvironment)
+	{
+		FGlobalShader::ModifyCompilationEnvironment(InParameters, OutEnvironment);
+		OutEnvironment.SetDefine(TEXT("MAX_NUM_DIMENSIONS"), FConvConstants::MAX_NUM_DIMENSIONS);
+	}
+
+	TArray<int32> FConvCS::GetOutputShape(TArrayView<const uint32> XShape, TArrayView<const uint32> WShape, EConvAutoPad AutoPad, TArrayView<const int32> Dilations, TArrayView<const int32> Strides, TArrayView<const int32> Pads)
+	{
+		check(XShape.Num() > 2);
+		check(WShape.Num() == XShape.Num());
+		check(Dilations.Num() == 0 || Dilations.Num() == WShape.Num() - 2);
+		check(Strides.Num() == 0 || Strides.Num() == WShape.Num() - 2);
+		check(Pads.Num() == 0 || Pads.Num() == 2 * (WShape.Num() - 2));
+
+		TArray<int32> Padding = GetPadding(XShape, WShape, AutoPad, Dilations, Strides, Pads);
+
+		TArray<int32> Result;
+		Result.SetNumUninitialized(XShape.Num());
+		Result[0] = XShape[0];
+		Result[1] = WShape[0];
+
+		int32 NumDimensions = XShape.Num() - 2;
+		for (int32 DimensionIndex = 0; DimensionIndex < NumDimensions; DimensionIndex++)
+		{
+			int32 PaddedSize = XShape[2 + DimensionIndex] + Padding[DimensionIndex] + Padding[NumDimensions + DimensionIndex];
+			int32 DilatedKernelSize = (DimensionIndex < Dilations.Num() ? Dilations[DimensionIndex] : 1) * (WShape[2 + DimensionIndex] - 1) + 1;
+			float PaddeSizeMinusDilatedKernelSize = (float)(PaddedSize - DilatedKernelSize);
+			int32 OutputSize = (int32)(PaddeSizeMinusDilatedKernelSize / (DimensionIndex < Strides.Num() ? (float)Strides[DimensionIndex] : 1.0) + 1);
+			Result[2 + DimensionIndex] = OutputSize;
+		}
+
+		return Result;
+	}
+
+	void FConvCS::FillInParameters(EConvGroupSize GroupSize, TArrayView<const uint32> XShape, TArrayView<const uint32> WShape, bool HasB,
+			EConvAutoPad AutoPad, int Group, TArrayView<const int32> Dilations, TArrayView<const int32> Strides, TArrayView<const int32> Pads, FConvCS::FParameters& Parameters)
+	{
+		check(XShape.Num() > 2);
+		check(WShape.Num() == XShape.Num());
+		check(Dilations.Num() == 0 || Dilations.Num() == WShape.Num() - 2);
+		check(Strides.Num() == 0 || Strides.Num() == WShape.Num() - 2);
+		check(Pads.Num() == 0 || Pads.Num() == 2 * (WShape.Num() - 2));
+		check(GetNumReadsPerThread(GroupSize, WShape, Dilations, Strides) >= 0)
+		check(WShape[0] > 0)
+		check(WShape[1] > 0)
+
+		int32 NumDimensions = XShape.Num() - 2;
+
+		TArray<int32> Padding = GetPadding(XShape, WShape, AutoPad, Dilations, Strides, Pads);
+		TArray<int32> GroupShape = GetGroupShape(GroupSize, NumDimensions);
+		TArray<int32> YShape = GetOutputShape(XShape, WShape, AutoPad, Dilations, Strides, Pads);
+		TArray<int32> GridShape = GetGridShape(YShape, GroupShape);
+		TArray<int32> XBlockShape = GetXBlockShape(GroupShape, WShape, Dilations, Strides);
+		
+		int32 GroupStride = 1;
+		int32 GroupThreadStride = 1;
+		int32 XBlockSize = 1;
+		int32 YMemoryStride = 1;
+		int32 XMemoryStride = 1;
+		int32 WChannelSize = 1;
+		for (int32 i = NumDimensions-1; i >= 0; i--)
+		{
+			int32 Stride = i < Strides.Num() ? Strides[i] : 1;
+			int32 Dilation = i < Dilations.Num() ? Dilations[i] : 1;
+			Parameters.Dilation_Stride_XBlockStartOffset_DilationXBlockStride[i] = FIntVector4(Dilation, Stride, -Padding[i], Dilation * XBlockSize);
+			Parameters.GroupStride_GroupShape_GroupThreadStride_StrideXBlockStride[i] = FIntVector4(GroupStride, GroupShape[i], GroupThreadStride, Stride * XBlockSize);
+			Parameters.YDimension_YMemoryStride_XDimension_XMemoryStride[i] = FIntVector4(YShape[2 + i], YMemoryStride, XShape[2 + i], XMemoryStride);
+			Parameters.XBlockStartStride_XBlockStride_WDimension_WDimensionDilationXBlockStride[i] = FIntVector4(GroupShape[i] * Stride, XBlockSize, WShape[2 + i], WShape[2 + i] * Dilation * XBlockSize);
+			Parameters.OneDiv_GroupStride_GroupThreadStride_XBlockStride[i] = FVector4f(1.0/((float)GroupStride), 1.0/((float)GroupThreadStride), 1.0/((float)XBlockSize), 0.0);
+
+			GroupStride *= GridShape[i];
+			GroupThreadStride *= GroupShape[i];
+			XBlockSize *= XBlockShape[i];
+			YMemoryStride *= YShape[2 + i];
+			XMemoryStride *= XShape[2 + i];
+			WChannelSize *= WShape[2 + i];
+		}
+		
+		Parameters.NumWChannels = WShape[1];
+
+		Parameters.YBatchStride = YShape[1] * YMemoryStride;
+		Parameters.YOutputKernelStride = YMemoryStride;
+		
+		Parameters.XBatchStride = XShape[1] * XMemoryStride;
+		Parameters.XChannelStride = XMemoryStride;
+
+		Parameters.XBlockSize = XBlockSize;
+
+		Parameters.NumChannelsPerBatch = FMath::Min((int32)((float)GroupThreadStride / (float)WChannelSize), (int32)WShape[1]);
+		check(Parameters.NumChannelsPerBatch > 0)
+		Parameters.NumChannelBatches = FMath::DivideAndRoundUp((int32)WShape[1], Parameters.NumChannelsPerBatch);
+		
+		Parameters.WOutputKernelStride = WShape[1] * WChannelSize;
+		Parameters.WChannelBatchSize = Parameters.NumChannelsPerBatch * WChannelSize;
+		Parameters.WChannelSize = WChannelSize;
+
+		Parameters.GroupsDivM = (float)Group / (float)WShape[0];
+	}
+
+	int32 FConvCS::GetNumReadsPerThread(EConvGroupSize GroupSize, TArrayView<const uint32> WShape, TArrayView<const int32> Dilations, TArrayView<const int32> Strides)
+	{
+		check(WShape.Num() > 2);
+		check(Dilations.Num() == 0 || Dilations.Num() == WShape.Num() - 2);
+		check(Strides.Num() == 0 || Strides.Num() == WShape.Num() - 2);
+
+		int32 NumDimensions = WShape.Num() - 2;
+
+		TArray<int32> GroupShape = GetGroupShape(GroupSize, NumDimensions);
+		int32 NumThreadsPerGroup = 1;
+		for (int32 i = 0; i < NumDimensions; i++)
+		{
+			NumThreadsPerGroup *= GroupShape[i];
+		}
+
+		TArray<int32> XBlockShape = GetXBlockShape(GroupShape, WShape, Dilations, Strides);
+		int32 NumXBlockElements = 1;
+		for (int32 i = 0; i < NumDimensions; i++)
+		{
+			NumXBlockElements *= XBlockShape[i];
+		}
+
+		int32 NumReads = FMath::DivideAndRoundUp(NumXBlockElements, NumThreadsPerGroup);
+		int32 NumReadsPow2 = FMath::Max(FMath::RoundToPositiveInfinity(FMath::Log2((float)NumReads)), FConvConstants::MIN_NUM_READS_PER_THREAD_POW2);
+
+		if(NumReadsPow2 <= FConvConstants::MAX_NUM_READS_PER_THREAD_POW2)
+		{
+			return NumReadsPow2;
+		}
+
+		return -1;
+	}
+
+	TArray<int32> FConvCS::GetGroupShape(EConvGroupSize GroupSize, int32 NumDimensions)
+	{
+		check(NumDimensions > 0);
+
+		int32 NumThreadsPerGroup = GetNumThreadsPerGroup(GroupSize);
+
+		int32 Power = (int32)FMath::Log2((float)NumThreadsPerGroup);
+		int32 MinPowerPerDim = (int32)((float)Power / (float)NumDimensions);
+		int32 PowerReminder = Power - NumDimensions * MinPowerPerDim;
+
+		TArray<int32> Result;
+		Result.Init((int32)FMath::Pow((float)2.0, (float)MinPowerPerDim), NumDimensions);
+		for (int32 i = 0; i < PowerReminder; i++)
+		{
+			Result[NumDimensions-1-i] *= 2;
+		}
+		return Result;
+	}
+
+	FIntVector FConvCS::GetGroupCount(TArrayView<const int32> YShape, TArrayView<const int32> GroupShape)
+	{
+		check(YShape.Num() > 2);
+		check(YShape.Num() == (GroupShape.Num() + 2));
+
+		int32 ThreadGroupCountValueX = 1;
+		for (int32 i = 2; i < YShape.Num(); i++)
+		{
+			ThreadGroupCountValueX *= FMath::DivideAndRoundUp(YShape[i], GroupShape[i-2]);
+		}
+
+		return FIntVector(ThreadGroupCountValueX, YShape[1], YShape[0]);
+	}
+
+	EConvGroupSize FConvCS::GetMinimalGroupSize(TArrayView<const int32> WShape)
+	{
+		int32 NumDimensions = WShape.Num() - 2;
+		int32 WChannelSize = 1;
+		for (int32 i = 0; i < NumDimensions; i++)
+		{
+			WChannelSize *= WShape[2 + i];
+		}
+		
+		for (int32 i = 0; i < (int32)EConvGroupSize::MAX; i++)
+		{
+			if (GetNumThreadsPerGroup((EConvGroupSize)i) >= WChannelSize) 
+			{
+				return (EConvGroupSize)i;
+			}
+		}
+		
+		return EConvGroupSize::MAX;
+	}
+
+	IMPLEMENT_GLOBAL_SHADER(FConvCS, "/NNX/ConvOp.usf", "main", SF_Compute);
+} // UE::NNIHlslShaders::Internal
\ No newline at end of file
diff --git a/Engine/Plugins/Experimental/NNX/Source/NNXHlslShaders/Private/NNIHlslShadersConvTransposeCS.cpp b/Engine/Plugins/Experimental/NNX/Source/NNXHlslShaders/Private/NNIHlslShadersConvTransposeCS.cpp
new file mode 100644
--- /dev/null
+++ b/Engine/Plugins/Experimental/NNX/Source/NNXHlslShaders/Private/NNIHlslShadersConvTransposeCS.cpp
@@ -0,0 +1,285 @@
+// Copyright Epic Games, Inc. All Rights Reserved.
+
+#include "NNIHlslShadersConvTransposeCS.h"
+
+namespace UE::NNIHlslShaders::Internal
+{
+	void FConvTransposeCS::ModifyCompilationEnvironment(const FGlobalShaderPermutationParameters& InParameters, FShaderCompilerEnvironment& OutEnvironment)
+	{
+		FGlobalShader::ModifyCompilationEnvironment(InParameters, OutEnvironment);
+		OutEnvironment.SetDefine(TEXT("MAX_NUM_STACK_DIMENSIONS"), FConvTransposeConstants::MAX_NUM_DIMENSIONS);
+	}
+
+	TArray<int32> FConvTransposeCS::GetOutputShape(TArrayView<const uint32> XShape, TArrayView<const uint32> WShape, EConvTransposeAutoPad AutoPad, TArrayView<const int32> Dilations, TArrayView<const int32> Strides, TArrayView<const int32> Pads, TArrayView<const int32> OutputPadding, int32 Group)
+	{
+		check(XShape.Num() > 2);
+		check(WShape.Num() == XShape.Num());
+		check(Dilations.Num() == 0 || Dilations.Num() == WShape.Num() - 2);
+		check(Strides.Num() == 0 || Strides.Num() == WShape.Num() - 2);
+		check(AutoPad != EConvTransposeAutoPad::NOTSET || Pads.Num() == 2 * (WShape.Num() - 2));
+
+		TArray<int32> Padding = GetPadding(WShape, AutoPad, Dilations, Strides, Pads, OutputPadding);
+
+		TArray<int32> Result;
+		Result.SetNumUninitialized(XShape.Num());
+		Result[0] = XShape[0];
+		Result[1] = WShape[1] * Group;
+
+		int32 NumDimensions = XShape.Num() - 2;
+		for (int32 DimensionIndex = 0; DimensionIndex < NumDimensions; DimensionIndex++)
+		{
+			Result[DimensionIndex + 2] = Strides[DimensionIndex] * (XShape[DimensionIndex + 2] - 1) + OutputPadding[DimensionIndex] +
+				((WShape[DimensionIndex + 2] - 1) * Dilations[DimensionIndex] + 1) -
+				Padding[DimensionIndex] - Padding[Strides.Num() + DimensionIndex];
+		}
+
+		return Result;
+	}
+
+	void FConvTransposeCS::FillInParameters(EConvTransposeGroupSize GroupSize, TArrayView<const uint32> XShape, TArrayView<const uint32> WShape, bool HasB, EConvTransposeAutoPad AutoPad, int32 Group, TArrayView<const int32> Dilations, TArrayView<const int32> Strides, TArrayView<const int32> Pads, TArrayView<const int32> OutputPadding, FConvTransposeCS::FParameters& Parameters)
+	{
+		check(XShape.Num() > 2);
+		check(WShape.Num() == XShape.Num());
+		check(Dilations.Num() == 0 || Dilations.Num() == WShape.Num() - 2);
+		check(Strides.Num() == 0 || Strides.Num() == WShape.Num() - 2);
+		check(AutoPad != EConvTransposeAutoPad::NOTSET || Pads.Num() == 2 * (WShape.Num() - 2));
+		check(GetNumReadsPerThread(GroupSize, WShape, Dilations, Strides) >= 0)
+			check(WShape[0] > 0)
+			check(WShape[1] > 0)
+
+			int32 NumDimensions = XShape.Num() - 2;
+
+		TArray<int32> Padding = GetPadding(WShape, AutoPad, Dilations, Strides, Pads, OutputPadding);
+		TArray<int32> GroupShape = GetGroupShape(GroupSize, NumDimensions);
+		TArray<int32> YShape = GetOutputShape(XShape, WShape, AutoPad, Dilations, Strides, Pads, OutputPadding, Group);
+		TArray<int32> GridShape = GetGridShape(YShape, GroupShape);
+		TArray<int32> XBlockShape = GetXBlockShape(GroupShape, WShape, Dilations, Strides);
+
+		int32 GroupStride = 1;
+		int32 GroupThreadStride = 1;
+		int32 XBlockSize = 1;
+		int32 YMemoryStride = 1;
+		int32 XMemoryStride = 1;
+		int32 WChannelSize = 1;
+		for (int32 i = NumDimensions - 1; i >= 0; i--)
+		{
+			int32 Stride = i < Strides.Num() ? Strides[i] : 1;
+			int32 Dilation = i < Dilations.Num() ? Dilations[i] : 1;
+			Parameters.Dilation_Stride_XBlockStartOffset_DilationXBlockStride[i] = FIntVector4(Dilation, Stride, Padding[i] + Dilations[i] * (1 - WShape[2 + i]), Dilation * XBlockSize);
+			Parameters.GroupStride_GroupShape_GroupThreadStride_StrideXBlockStride[i] = FIntVector4(GroupStride, GroupShape[i], GroupThreadStride, XBlockSize);
+			Parameters.YDimension_YMemoryStride_XDimension_XMemoryStride[i] = FIntVector4(YShape[2 + i], YMemoryStride, XShape[2 + i], XMemoryStride);
+			Parameters.XBlockStartStride_XBlockStride_WDimension_WDimensionDilationXBlockStride[i] = FIntVector4(GroupShape[i], XBlockSize, WShape[2 + i], WShape[2 + i] * Dilation * XBlockSize);
+			Parameters.OneDiv_GroupStride_GroupThreadStride_OneDivStride[i] = FVector4f(1.0 / ((float)GroupStride), 1.0 / ((float)GroupThreadStride), 1.0 / ((float)XBlockSize), 1.0 / ((float)Stride));
+
+			GroupStride *= GridShape[i];
+			GroupThreadStride *= GroupShape[i];
+			XBlockSize *= XBlockShape[i];
+			YMemoryStride *= YShape[2 + i];
+			XMemoryStride *= XShape[2 + i];
+			WChannelSize *= WShape[2 + i];
+		}
+
+		Parameters.NumWChannels = WShape[0];
+		Parameters.NumOutChannelsDivGroup = WShape[1];
+
+		Parameters.YBatchStride = YShape[1] * YMemoryStride;
+		Parameters.YOutputKernelStride = YMemoryStride;
+
+		Parameters.XBatchStride = XShape[1] * XMemoryStride;
+		Parameters.XChannelStride = XMemoryStride;
+
+		Parameters.XBlockSize = XBlockSize;
+
+		Parameters.NumChannelsPerBatch = FMath::Min((int32)((float)GroupThreadStride / (float)WChannelSize), (int32)WShape[0]);
+		check(Parameters.NumChannelsPerBatch > 0)
+			Parameters.NumChannelBatches = FMath::DivideAndRoundUp((int32)WShape[0], Parameters.NumChannelsPerBatch);
+
+		Parameters.WOutputKernelStride = WShape[1] * WChannelSize;
+		Parameters.WChannelBatchSize = Parameters.NumChannelsPerBatch * WShape[1] * WChannelSize;
+		Parameters.WChannelSize = WChannelSize;
+
+		Parameters.GroupsDivM = 1.0 / ((float)WShape[1] * Group);
+		Parameters.OneDivGroup = 1.0 / (float)Group;
+	}
+
+	int32 FConvTransposeCS::GetNumReadsPerThread(EConvTransposeGroupSize GroupSize, TArrayView<const uint32> WShape, TArrayView<const int32> Dilations, TArrayView<const int32> Strides)
+	{
+		check(WShape.Num() > 2);
+		check(Dilations.Num() == 0 || Dilations.Num() == WShape.Num() - 2);
+		check(Strides.Num() == 0 || Strides.Num() == WShape.Num() - 2);
+
+		int32 NumDimensions = WShape.Num() - 2;
+
+		TArray<int32> GroupShape = GetGroupShape(GroupSize, NumDimensions);
+		int32 NumThreadsPerGroup = 1;
+		for (int32 i = 0; i < NumDimensions; i++)
+		{
+			NumThreadsPerGroup *= GroupShape[i];
+		}
+
+		TArray<int32> XBlockShape = GetXBlockShape(GroupShape, WShape, Dilations, Strides);
+		int32 NumXBlockElements = 1;
+		for (int32 i = 0; i < NumDimensions; i++)
+		{
+			NumXBlockElements *= XBlockShape[i];
+		}
+
+		int32 NumReads = FMath::DivideAndRoundUp(NumXBlockElements, NumThreadsPerGroup);
+		int32 NumReadsPow2 = FMath::Max(FMath::RoundToPositiveInfinity(FMath::Log2((float)NumReads)), FConvTransposeConstants::MIN_NUM_READS_PER_THREAD_POW2);
+
+		if (NumReadsPow2 <= FConvTransposeConstants::MAX_NUM_READS_PER_THREAD_POW2)
+		{
+			return NumReadsPow2;
+		}
+
+		return -1;
+	}
+
+	TArray<int32> FConvTransposeCS::GetGroupShape(EConvTransposeGroupSize GroupSize, int32 NumDimensions)
+	{
+		check(NumDimensions > 0);
+
+		int32 NumThreadsPerGroup = GetNumThreadsPerGroup(GroupSize);
+
+		int32 Power = (int32)FMath::Log2((float)NumThreadsPerGroup);
+		int32 MinPowerPerDim = (int32)((float)Power / (float)NumDimensions);
+		int32 PowerReminder = Power - NumDimensions * MinPowerPerDim;
+
+		TArray<int32> Result;
+		Result.Init((int32)FMath::Pow((float)2.0, (float)MinPowerPerDim), NumDimensions);
+		for (int32 i = 0; i < PowerReminder; i++)
+		{
+			Result[NumDimensions - 1 - i] *= 2;
+		}
+
+		return Result;
+	}
+
+	FIntVector FConvTransposeCS::GetGroupCount(TArrayView<const int32> YShape, TArrayView<const int32> GroupShape)
+	{
+		check(YShape.Num() > 2);
+		check(YShape.Num() == (GroupShape.Num() + 2));
+
+		int32 ThreadGroupCountValueX = 1;
+		for (int32 i = 2; i < YShape.Num(); i++)
+		{
+			ThreadGroupCountValueX *= FMath::DivideAndRoundUp(YShape[i], GroupShape[i - 2]);
+		}
+
+		return FIntVector(ThreadGroupCountValueX, YShape[1], YShape[0]);
+	}
+
+	EConvTransposeGroupSize FConvTransposeCS::GetMinimalGroupSize(TArrayView<const uint32> WShape)
+	{
+		int32 NumDimensions = WShape.Num() - 2;
+		int32 WChannelSize = 1;
+		for (int32 i = 0; i < NumDimensions; i++)
+		{
+			WChannelSize *= WShape[2 + i];
+		}
+
+		for (int32 i = 0; i < (int32)EConvTransposeGroupSize::MAX; i++)
+		{
+			if (GetNumThreadsPerGroup((EConvTransposeGroupSize)i) >= WChannelSize)
+			{
+				return (EConvTransposeGroupSize)i;
+			}
+		}
+
+		return EConvTransposeGroupSize::MAX;
+	}
+
+	TArray<int32> FConvTransposeCS::GetXBlockShape(TArrayView<const int32> GroupShape, TArrayView<const uint32> WShape, TArrayView<const int32> Dilations, TArrayView<const int32> Strides)
+	{
+		check(WShape.Num() > 2);
+		check(GroupShape.Num() == WShape.Num() - 2);
+		check(Dilations.Num() == 0 || Dilations.Num() == GroupShape.Num());
+		check(Strides.Num() == 0 || Strides.Num() == GroupShape.Num());
+
+		TArray<int32> Result;
+		Result.SetNumUninitialized(GroupShape.Num());
+		for (int32 i = 0; i < GroupShape.Num(); i++)
+		{
+			int32 DilatedKernelSize = (i < Dilations.Num() ? Dilations[i] : 1) * (WShape[2 + i] - 1) + 1;
+			Result[i] = DilatedKernelSize + (GroupShape[i] - 1);
+		}
+		return Result;
+	}
+
+	TArray<int32> FConvTransposeCS::GetPadding(TArrayView<const uint32> WShape, EConvTransposeAutoPad AutoPad, TArrayView<const int32> Dilations, TArrayView<const int32> Strides, TArrayView<const int32> Pads, TArrayView<const int32> OutputPadding)
+	{
+		check(WShape.Num() > 2);
+		check(Dilations.Num() == 0 || Dilations.Num() == WShape.Num() - 2);
+		check(Strides.Num() == 0 || Strides.Num() == WShape.Num() - 2);
+		check(AutoPad != EConvTransposeAutoPad::NOTSET || Pads.Num() == 2 * (WShape.Num() - 2));
+
+		int32 NumDimensions = WShape.Num() - 2;
+		TArray<int32> Result;
+		Result.Init(0, 2 * NumDimensions);
+
+		if (AutoPad == EConvTransposeAutoPad::NOTSET)
+		{
+			return TArray<int32>{Pads};
+		}
+		else if (AutoPad == EConvTransposeAutoPad::VALID)
+		{
+			return Result;
+		}
+
+		for (int32 DimensionIndex = 0; DimensionIndex < NumDimensions; DimensionIndex++)
+		{
+			int32 TotalPad = (WShape[DimensionIndex + 2] - 1) * Dilations[DimensionIndex] - Strides[DimensionIndex] + OutputPadding[DimensionIndex] + 1;
+
+			if (AutoPad == EConvTransposeAutoPad::SAME_LOWER)
+			{
+				Result[DimensionIndex] = (TotalPad + 1) / 2;
+			}
+			else
+			{
+				Result[DimensionIndex] = TotalPad / 2;
+			}
+			Result[NumDimensions + DimensionIndex] = TotalPad - Result[DimensionIndex];
+		}
+
+		return Result;
+	}
+
+	int32 FConvTransposeCS::GetNumThreadsPerGroup(EConvTransposeGroupSize GroupSize)
+	{
+		int32 NumThreadsPerGroup = 128;
+		switch (GroupSize)
+		{
+		case EConvTransposeGroupSize::Size128:
+			NumThreadsPerGroup = 128;
+			break;
+		case EConvTransposeGroupSize::Size256:
+			NumThreadsPerGroup = 256;
+			break;
+		case EConvTransposeGroupSize::Size512:
+			NumThreadsPerGroup = 512;
+			break;
+		default:
+			NumThreadsPerGroup = 128;
+			break;
+		}
+		check(FMath::Log2((float)NumThreadsPerGroup) == FMath::Floor(FMath::Log2((float)NumThreadsPerGroup)));
+		return NumThreadsPerGroup;
+	}
+
+	TArray<int32> FConvTransposeCS::GetGridShape(TArrayView<const int32> YShape, TArrayView<const int32> GroupShape)
+	{
+		check(YShape.Num() > 2);
+		check(YShape.Num() == (GroupShape.Num() + 2));
+
+		TArray<int32> Result;
+		Result.SetNumUninitialized(YShape.Num() - 2);
+		for (int32 i = 2; i < YShape.Num(); i++)
+		{
+			Result[i - 2] = FMath::DivideAndRoundUp(YShape[i], GroupShape[i - 2]);
+		}
+
+		return Result;
+	}
+
+	IMPLEMENT_GLOBAL_SHADER(FConvTransposeCS, "/NNX/ConvTransposeOp.usf", "main", SF_Compute);
+} // UE::NNIHlslShaders::Internal
\ No newline at end of file
diff --git a/Engine/Plugins/Experimental/NNX/Source/NNXHlslShaders/Private/NNXConvCS.cpp b/Engine/Plugins/Experimental/NNX/Source/NNXHlslShaders/Private/NNXConvCS.cpp
deleted file mode 100644
diff --git a/Engine/Plugins/Experimental/NNX/Source/NNXHlslShaders/Private/NNXConvTransposeCS.cpp b/Engine/Plugins/Experimental/NNX/Source/NNXHlslShaders/Private/NNXConvTransposeCS.cpp
deleted file mode 100644
diff --git a/Engine/Plugins/Experimental/NNX/Source/NNXHlslShaders/Public/NNXConvCS.h b/Engine/Plugins/Experimental/NNX/Source/NNXHlslShaders/Public/NNXConvCS.h
deleted file mode 100644
diff --git a/Engine/Plugins/Experimental/NNX/Source/NNXHlslShaders/Public/NNXConvTransposeCS.h b/Engine/Plugins/Experimental/NNX/Source/NNXHlslShaders/Public/NNXConvTransposeCS.h
deleted file mode 100644
diff --git a/Engine/Plugins/Experimental/NNX/Source/NNXRuntimeRDG/Private/Hlsl/NNIRuntimeRDGConv.cpp b/Engine/Plugins/Experimental/NNX/Source/NNXRuntimeRDG/Private/Hlsl/NNIRuntimeRDGConv.cpp
--- a/Engine/Plugins/Experimental/NNX/Source/NNXRuntimeRDG/Private/Hlsl/NNIRuntimeRDGConv.cpp
+++ b/Engine/Plugins/Experimental/NNX/Source/NNXRuntimeRDG/Private/Hlsl/NNIRuntimeRDGConv.cpp
@@ -1,15 +1,17 @@
 // Copyright Epic Games, Inc. All Rights Reserved.
 
 #include "NNIRuntimeRDGConv.h"
-#include "NNXConvCS.h"
+#include "NNIHlslShadersConvCS.h"
 #include "NNXRuntimeHLSLHelper.h"
 
 namespace UE::NNIRuntimeRDG::Private::Hlsl
 {
-	DECLARE_GPU_STAT_NAMED(FMLHLSLOperatorConv, TEXT("NNI.Operator.Hlsl.Conv"));
+	DECLARE_GPU_STAT_NAMED(FNNIOperatorConv, TEXT("NNI.Operator.Hlsl.Conv"));
+
+	using EConvAutoPad = UE::NNIHlslShaders::Internal::EConvAutoPad;
 
 	/**
-	 * Convolution ML operator implementation
+	 * Convolution operator implementation
 	 */
 	class FConv : public NNX::FMLOperatorHlsl
 	{
@@ -26,13 +28,13 @@ namespace UE::NNIRuntimeRDG::Private::Hlsl
 
 		FConv() {}
 
-		NNX::FMLTensorDesc Input;
-		NNX::FMLTensorDesc Weights;
-		NNX::FMLTensorDesc Bias;
-		NNX::FMLTensorDesc Output;
+		NNX::FMLTensorDesc Input = {};
+		NNX::FMLTensorDesc Weights = {};
+		NNX::FMLTensorDesc Bias = {};
+		NNX::FMLTensorDesc Output = {};
 
-		int NumDimensions;
-		bool HasBias;
+		int NumDimensions = 0;
+		bool HasBias = false;
 
 		// Hard coded parameters, until we accept them from JSON
 		int32 Group = 1;
@@ -71,14 +73,16 @@ namespace UE::NNIRuntimeRDG::Private::Hlsl
 
 		virtual void Dispatch(FRDGBuilder& GraphBuilder, TArrayView<const NNX::FMLTensorBinding> InInputBindings, TArrayView<const NNX::FMLTensorBinding> OutOutputBindings) override
 		{
+			using namespace UE::NNIHlslShaders::Internal;
+
 			constexpr EConvAlgorithm Algorithm = EConvAlgorithm::SharedMemory;
 			constexpr EConvGroupSize GroupSize = EConvGroupSize::Size256;
 
-			TArray<int32> OutputShape = FMLConvCS::GetOutputShape(Input.Shape, Weights.Shape, AutoPad, Dilations, Strides, Pads);
+			TArray<int32> OutputShape = FConvCS::GetOutputShape(Input.Shape, Weights.Shape, AutoPad, Dilations, Strides, Pads);
 
 			// Set parameters
-			FMLConvCS::FParameters* Params = GraphBuilder.AllocParameters<FMLConvCS::FParameters>();
-			FMLConvCS::FillInParameters(GroupSize, Input.Shape, Weights.Shape, HasBias, AutoPad, Group, Dilations,Strides, Pads, *Params);
+			FConvCS::FParameters* Params = GraphBuilder.AllocParameters<FConvCS::FParameters>();
+			FConvCS::FillInParameters(GroupSize, Input.Shape, Weights.Shape, HasBias, AutoPad, Group, Dilations,Strides, Pads, *Params);
 			Params->X = GraphBuilder.CreateSRV(FRDGBufferSRVDesc(InInputBindings[0].Buffer, PF_R32_FLOAT));
 			Params->W = GraphBuilder.CreateSRV(FRDGBufferSRVDesc(InInputBindings[1].Buffer, PF_R32_FLOAT));
 			if (InInputBindings.Num() == 3) {
@@ -86,25 +90,25 @@ namespace UE::NNIRuntimeRDG::Private::Hlsl
 			}
 			Params->Y = GraphBuilder.CreateUAV(FRDGBufferUAVDesc(OutOutputBindings[0].Buffer, PF_R32_FLOAT));
 
-			FMLConvCS::FPermutationDomain PermutationVector;
+			FConvCS::FPermutationDomain PermutationVector;
 
-			PermutationVector.Set<FMLConvCS::FConvAlgorithm>(Algorithm);
-			PermutationVector.Set<FMLConvCS::FConvGroupSize>(GroupSize);
-			PermutationVector.Set<FMLConvCS::FConvNumDimensions>(NumDimensions);
-			PermutationVector.Set<FMLConvCS::FConvNumReadsPerThread>(FMLConvCS::GetNumReadsPerThread(GroupSize, Weights.Shape, Dilations, Strides));
-			PermutationVector.Set<FMLConvCS::FConvHasB>(HasBias);
-			TShaderMapRef<FMLConvCS> ComputeShader(GetGlobalShaderMap(GMaxRHIFeatureLevel), PermutationVector);
+			PermutationVector.Set<FConvCS::FConvAlgorithm>(Algorithm);
+			PermutationVector.Set<FConvCS::FConvGroupSize>(GroupSize);
+			PermutationVector.Set<FConvCS::FConvNumDimensions>(NumDimensions);
+			PermutationVector.Set<FConvCS::FConvNumReadsPerThread>(FConvCS::GetNumReadsPerThread(GroupSize, Weights.Shape, Dilations, Strides));
+			PermutationVector.Set<FConvCS::FConvHasB>(HasBias);
+			TShaderMapRef<FConvCS> ComputeShader(GetGlobalShaderMap(GMaxRHIFeatureLevel), PermutationVector);
 
-			RDG_EVENT_SCOPE(GraphBuilder, "FML.HLSL.Operator.Conv");
-			RDG_GPU_STAT_SCOPE(GraphBuilder, FMLHLSLOperatorConv);
+			RDG_EVENT_SCOPE(GraphBuilder, "NNI.Operator.Hlsl.Conv");
+			RDG_GPU_STAT_SCOPE(GraphBuilder, FNNIOperatorConv);
 
 			FComputeShaderUtils::AddPass(
 				GraphBuilder,
-				RDG_EVENT_NAME("FML.HLSL.Operator.Conv.Dispatch"),
+				RDG_EVENT_NAME("NNI.Operator.Hlsl.Conv.Dispatch"),
 				ERDGPassFlags::Compute | ERDGPassFlags::NeverCull,
 				ComputeShader,
 				Params,
-				FMLConvCS::GetGroupCount(OutputShape, FMLConvCS::GetGroupShape(GroupSize, NumDimensions)));
+				FConvCS::GetGroupCount(OutputShape, FConvCS::GetGroupShape(GroupSize, NumDimensions)));
 		}
 	};
 
@@ -115,4 +119,4 @@ namespace UE::NNIRuntimeRDG::Private::Hlsl
 		return true;
 	}
 
-} // NNX
+} // UE::NNIRuntimeRDG::Private::Hlsl
diff --git a/Engine/Plugins/Experimental/NNX/Source/NNXRuntimeRDG/Private/Hlsl/NNIRuntimeRDGConvTranspose.cpp b/Engine/Plugins/Experimental/NNX/Source/NNXRuntimeRDG/Private/Hlsl/NNIRuntimeRDGConvTranspose.cpp
--- a/Engine/Plugins/Experimental/NNX/Source/NNXRuntimeRDG/Private/Hlsl/NNIRuntimeRDGConvTranspose.cpp
+++ b/Engine/Plugins/Experimental/NNX/Source/NNXRuntimeRDG/Private/Hlsl/NNIRuntimeRDGConvTranspose.cpp
@@ -1,13 +1,15 @@
 // Copyright Epic Games, Inc. All Rights Reserved.
 
 #include "NNIRuntimeRDGConvTranspose.h"
-#include "NNXConvTransposeCS.h"
+#include "NNIHlslShadersConvTransposeCS.h"
 #include "NNXRuntimeHLSLHelper.h"
 
 namespace UE::NNIRuntimeRDG::Private::Hlsl
 {
 	DECLARE_GPU_STAT_NAMED(FNNIOperatorConvTranspose, TEXT("NNI.Operator.Hlsl.ConvTranspose"));
 
+	using EConvTransposeAutoPad = UE::NNIHlslShaders::Internal::EConvTransposeAutoPad;
+
 	/**
 	 * ConvTranspose operator implementation
 	 */
@@ -73,6 +75,8 @@ namespace UE::NNIRuntimeRDG::Private::Hlsl
 
 		virtual void Dispatch(FRDGBuilder& GraphBuilder, TArrayView<const NNX::FMLTensorBinding> InInputBindings, TArrayView<const NNX::FMLTensorBinding> OutOutputBindings) override
 		{
+			using namespace UE::NNIHlslShaders::Internal;
+
 			constexpr EConvTransposeAlgorithm Algorithm = EConvTransposeAlgorithm::SharedMemory;
 			constexpr EConvTransposeGroupSize GroupSize = EConvTransposeGroupSize::Size256;
 
@@ -117,4 +121,4 @@ namespace UE::NNIRuntimeRDG::Private::Hlsl
 		return true;
 	}
 
-} // NNX
+} // UE::NNIRuntimeRDG::Private::Hlsl
