NNI - Add namespaces to GEMM and MatMul

diff --git a/Engine/Plugins/Experimental/NNX/Source/NNXHlslShaders/Internal/NNIHlslShadersGatherCS.h b/Engine/Plugins/Experimental/NNX/Source/NNXHlslShaders/Internal/NNIHlslShadersGatherCS.h
new file mode 100644
--- /dev/null
+++ b/Engine/Plugins/Experimental/NNX/Source/NNXHlslShaders/Internal/NNIHlslShadersGatherCS.h
@@ -0,0 +1,49 @@
+// Copyright Epic Games, Inc. All Rights Reserved.
+
+#pragma once
+
+#include "GlobalShader.h"
+#include "ShaderParameterUtils.h"
+#include "RenderGraphUtils.h"
+#include "NNXTypes.h"
+
+namespace UE::NNIHlslShaders::Internal
+{
+	class FGatherConstants
+	{
+	public:
+
+		static const int32 MAX_NUM_DIMENSIONS{ 8 };
+		static const int32 NUM_GROUP_THREADS{ 256 };
+	};
+
+	// template <typename DataElementType, typename IndicesElementType>
+	class NNXHLSLSHADERS_API TGatherCS : public FGlobalShader
+	{
+		DECLARE_GLOBAL_SHADER(TGatherCS);
+		SHADER_USE_PARAMETER_STRUCT(TGatherCS, FGlobalShader)
+
+		class FGatherNumOutputDimensions : SHADER_PERMUTATION_RANGE_INT("NUM_OUTPUT_DIMENSIONS", 1, FGatherConstants::MAX_NUM_DIMENSIONS);
+		using FPermutationDomain = TShaderPermutationDomain<FGatherNumOutputDimensions>;
+
+	public:
+
+		BEGIN_SHADER_PARAMETER_STRUCT(FParameters, )
+			SHADER_PARAMETER(int32, Axis)
+			SHADER_PARAMETER(int32, OutputSize)
+			SHADER_PARAMETER(int32, NumDataDimensions)
+			SHADER_PARAMETER(int32, NumIndicesDimensions)
+			SHADER_PARAMETER_ARRAY(FIntVector4, DataStride_IndicesStride_OutputStride, [FGatherConstants::MAX_NUM_DIMENSIONS])
+			SHADER_PARAMETER_ARRAY(FVector4f, OneDivDataStride_OneDivIndicesStride_OneDivOutputStride, [FGatherConstants::MAX_NUM_DIMENSIONS])
+			SHADER_PARAMETER_RDG_BUFFER_SRV(Buffer<float>, Data)
+			SHADER_PARAMETER_RDG_BUFFER_SRV(Buffer<int32>, Indices)
+			SHADER_PARAMETER_RDG_BUFFER_UAV(RWBuffer<float>, Output)
+		END_SHADER_PARAMETER_STRUCT()
+
+		static void ModifyCompilationEnvironment(const FGlobalShaderPermutationParameters& InParameters, FShaderCompilerEnvironment& OutEnvironment);
+
+		static void FillInParameters(int32 Axis, const NNX::FMLTensorDesc& Data, const NNX::FMLTensorDesc& Indices, FParameters& Parameters);
+
+		static FIntVector GetGroupCount(const FParameters& Parameters);
+	};
+} // UE::NNIHlslShaders::Internal
\ No newline at end of file
diff --git a/Engine/Plugins/Experimental/NNX/Source/NNXHlslShaders/Internal/NNIHlslShadersGemmCS.h b/Engine/Plugins/Experimental/NNX/Source/NNXHlslShaders/Internal/NNIHlslShadersGemmCS.h
new file mode 100644
--- /dev/null
+++ b/Engine/Plugins/Experimental/NNX/Source/NNXHlslShaders/Internal/NNIHlslShadersGemmCS.h
@@ -0,0 +1,89 @@
+// Copyright Epic Games, Inc. All Rights Reserved.
+
+#pragma once
+
+#include "GlobalShader.h"
+#include "ShaderParameterUtils.h"
+#include "RenderGraphUtils.h"
+#include "NNXTypes.h"
+
+namespace UE::NNIHlslShaders::Internal
+{
+	enum class EGemmCScalar : uint8
+	{
+		No = 0,
+		Yes,
+		NoBias,
+		MAX
+	};
+
+	enum class EGemmAlgorithm : uint8
+	{
+		Simple8x8 = 0,
+		Simple16x16,
+		Simple32x32,
+		Simple256x1,
+		SharedMemory8x8,
+		SharedMemory16x16,
+		SharedMemory32x32,
+		MultiWrite1x16,
+		MultiWrite2x16,
+		MultiWrite1x32,
+		MultiWrite2x32,
+		MultiWrite4x32,
+		MultiWrite2x64,
+		MultiWrite4x64,
+		MultiWrite8x64,
+		MAX
+	};
+
+	class FGemmConstants
+	{
+	public:
+
+		static const int32 MAX_NUM_STACK_DIMENSIONS{8};
+	};
+
+	class NNXHLSLSHADERS_API TGemmCS : public FGlobalShader
+	{
+		DECLARE_GLOBAL_SHADER(TGemmCS);
+		SHADER_USE_PARAMETER_STRUCT(TGemmCS, FGlobalShader)
+
+		class FGemmCScalar : SHADER_PERMUTATION_ENUM_CLASS("C_SCALAR", EGemmCScalar);
+		class FGemmAlgorithm : SHADER_PERMUTATION_ENUM_CLASS("ALGORITHM", EGemmAlgorithm);
+		class FGemmNumStackDimensions : SHADER_PERMUTATION_RANGE_INT("NUM_STACK_DIMENSIONS", 0, FGemmConstants::MAX_NUM_STACK_DIMENSIONS);
+		using FPermutationDomain = TShaderPermutationDomain<FGemmCScalar, FGemmAlgorithm, FGemmNumStackDimensions>;
+
+	public:
+
+		BEGIN_SHADER_PARAMETER_STRUCT(FParameters, )
+			SHADER_PARAMETER(float, Alpha)
+			SHADER_PARAMETER(float, Beta)
+			SHADER_PARAMETER(int32, TransA)
+			SHADER_PARAMETER(int32, TransB)
+			SHADER_PARAMETER(uint32, M)
+			SHADER_PARAMETER(uint32, N)
+			SHADER_PARAMETER(uint32, K)
+			SHADER_PARAMETER(uint32, MxK)
+			SHADER_PARAMETER(uint32, KxN)
+			SHADER_PARAMETER(uint32, MxN)
+			SHADER_PARAMETER(uint32, CWidth)
+			SHADER_PARAMETER(uint32, CHeight)
+			SHADER_PARAMETER(float, CScalar)
+			SHADER_PARAMETER_ARRAY(FUint32Vector4, StackShapeA_StackShapeB_StackStrideA_StackStrideB, [FGemmConstants::MAX_NUM_STACK_DIMENSIONS])
+			SHADER_PARAMETER_RDG_BUFFER_SRV(Buffer<float>, A)
+			SHADER_PARAMETER_RDG_BUFFER_SRV(Buffer<float>, B)
+			SHADER_PARAMETER_RDG_BUFFER_SRV(Buffer<float>, C)
+			SHADER_PARAMETER_RDG_BUFFER_UAV(RWBuffer<float>, Y)
+		END_SHADER_PARAMETER_STRUCT()
+
+		static void ModifyCompilationEnvironment(const FGlobalShaderPermutationParameters& InParameters, FShaderCompilerEnvironment& OutEnvironment);
+
+		static void FillInParameters(float Alpha, float Beta, int32 TransA, int32 TransB, const NNX::FMLTensorDesc &InputA, const NNX::FMLTensorDesc &InputB, const NNX::FMLTensorDesc &InputC, float CScalar, FParameters& Parameters);
+
+		static void FillInParametersMatMul(const NNX::FMLTensorDesc &InputA, const NNX::FMLTensorDesc &InputB, FParameters& Parameters);
+
+		static FIntVector GetGroupCount(const FParameters& Parameters, EGemmAlgorithm Algorithm, int32 NumStackDimensions);
+		static EGemmAlgorithm GetAlgorithm(const FParameters& Parameters);
+	};
+} // UE::NNIHlslShaders::Internal
\ No newline at end of file
diff --git a/Engine/Plugins/Experimental/NNX/Source/NNXHlslShaders/Private/NNIHlslShadersGatherCS.cpp b/Engine/Plugins/Experimental/NNX/Source/NNXHlslShaders/Private/NNIHlslShadersGatherCS.cpp
new file mode 100644
--- /dev/null
+++ b/Engine/Plugins/Experimental/NNX/Source/NNXHlslShaders/Private/NNIHlslShadersGatherCS.cpp
@@ -0,0 +1,81 @@
+// Copyright Epic Games, Inc. All Rights Reserved.
+
+#include "NNIHlslShadersGatherCS.h"
+
+namespace UE::NNIHlslShaders::Internal
+{
+	// template <>
+	void TGatherCS::ModifyCompilationEnvironment(const FGlobalShaderPermutationParameters& InParameters, FShaderCompilerEnvironment& OutEnvironment)
+	{
+		FGlobalShader::ModifyCompilationEnvironment(InParameters, OutEnvironment);
+		OutEnvironment.SetDefine(TEXT("DATA_TYPE"), 0);
+		OutEnvironment.SetDefine(TEXT("INDICES_TYPE"), 0);
+		OutEnvironment.SetDefine(TEXT("MAX_NUM_DIMENSIONS"), FGatherConstants::MAX_NUM_DIMENSIONS);
+		OutEnvironment.SetDefine(TEXT("NUM_GROUP_THREADS"), FGatherConstants::NUM_GROUP_THREADS);
+	}
+
+	//template <typename DataElementType, typename IndicesElementType>
+	void TGatherCS::FillInParameters(int32 Axis, const NNX::FMLTensorDesc& Data, const NNX::FMLTensorDesc& Indices, FParameters& Parameters)
+	{
+		Parameters.Axis = Axis;
+
+		TArray<int32> OutputShape;
+		OutputShape.SetNum(Indices.Shape.Num() + (Data.Shape.Num() - 1));
+		int32 index = 0;
+		Parameters.OutputSize = 1;
+		for (int32 i = 0; i < Axis; i++, index++)
+		{
+			OutputShape[index] = Data.Shape[i];
+			Parameters.OutputSize *= OutputShape[index];
+		}
+		for (int32 i = 0; i < Indices.Shape.Num(); i++, index++)
+		{
+			OutputShape[index] = Indices.Shape[i];
+			Parameters.OutputSize *= OutputShape[index];
+		}
+		for (int32 i = (Axis + 1); i < Data.Shape.Num(); i++, index++)
+		{
+			OutputShape[index] = Data.Shape[i];
+			Parameters.OutputSize *= OutputShape[index];
+		}
+
+		Parameters.NumDataDimensions = Data.Shape.Num();
+		Parameters.NumIndicesDimensions = Indices.Shape.Num();
+
+		int32 DataStride = 1;
+		for (int32 i = Data.Shape.Num() - 1; i >= 0; i--)
+		{
+			Parameters.DataStride_IndicesStride_OutputStride[i].X = DataStride;
+			Parameters.OneDivDataStride_OneDivIndicesStride_OneDivOutputStride[i].X = 1.0 / (float)DataStride;
+			DataStride *= Data.Shape[i];
+		}
+
+		int32 IndicesStride = 1;
+		for (int32 i = Indices.Shape.Num() - 1; i >= 0; i--)
+		{
+			Parameters.DataStride_IndicesStride_OutputStride[i].Y = IndicesStride;
+			Parameters.OneDivDataStride_OneDivIndicesStride_OneDivOutputStride[i].Y = 1.0 / (float)IndicesStride;
+			IndicesStride *= Indices.Shape[i];
+		}
+
+		int32 OutputStride = 1;
+		for (int32 i = OutputShape.Num() - 1; i >= 0; i--)
+		{
+			Parameters.DataStride_IndicesStride_OutputStride[i].Z = OutputStride;
+			Parameters.OneDivDataStride_OneDivIndicesStride_OneDivOutputStride[i].Z = 1.0 / (float)OutputStride;
+			OutputStride *= OutputShape[i];
+		}
+	}
+
+	// template <typename DataElementType, typename IndicesElementType>
+	FIntVector TGatherCS::GetGroupCount(const FParameters& Parameters)
+	{
+		return FIntVector(FMath::DivideAndRoundUp(Parameters.OutputSize, FGatherConstants::NUM_GROUP_THREADS), 1, 1);
+	}
+
+	//typedef TGatherCS<float, int32> FGatherFloatInt32CS;
+	//IMPLEMENT_SHADER_TYPE(template<>, TGatherCS, TEXT("/NNX/GatherOp.usf"), TEXT("main"), SF_Compute);
+	IMPLEMENT_GLOBAL_SHADER(TGatherCS, "/NNX/GatherOp.usf", "main", SF_Compute);
+	//template class TGatherCS<float, int32>;
+
+} // UE::NNIHlslShaders::Internal
\ No newline at end of file
diff --git a/Engine/Plugins/Experimental/NNX/Source/NNXHlslShaders/Private/NNIHlslShadersGemmCS.cpp b/Engine/Plugins/Experimental/NNX/Source/NNXHlslShaders/Private/NNIHlslShadersGemmCS.cpp
new file mode 100644
--- /dev/null
+++ b/Engine/Plugins/Experimental/NNX/Source/NNXHlslShaders/Private/NNIHlslShadersGemmCS.cpp
@@ -0,0 +1,190 @@
+// Copyright Epic Games, Inc. All Rights Reserved.
+
+#include "NNIHlslShadersGemmCS.h"
+#include "NNXCore.h"
+#include "Algo/Accumulate.h"
+
+namespace UE::NNIHlslShaders::Internal
+{
+	struct GemmMatrixParameters
+	{
+		uint32 M = 0;
+		uint32 N = 0;
+		uint32 K = 0;
+		TArray<uint32> StackShapeA;
+		TArray<uint32> StackShapeB;
+		TArray<uint32> StackStrideA;
+		TArray<uint32> StackStrideB;
+	};
+
+	static GemmMatrixParameters GetMatrixParametersMatMul(const NNX::FMLTensorDesc& InputA, const NNX::FMLTensorDesc& InputB)
+	{
+		check(InputA.Shape.Num() != 0);
+		check(InputB.Shape.Num() != 0);
+
+		GemmMatrixParameters Result;
+		int32 NumStackDimensionsA = FMath::Max((int32)InputA.Shape.Num() - 2, 0);
+		int32 NumStackDimensionsB = FMath::Max((int32)InputB.Shape.Num() - 2, 0);
+		int32 NumStackDimensions = FMath::Max(NumStackDimensionsA, NumStackDimensionsB);
+
+		check(NumStackDimensions <= FGemmConstants::MAX_NUM_STACK_DIMENSIONS);
+
+		// Check matrix stack dimensions
+		if (NumStackDimensionsA > 0 && NumStackDimensionsB > 0)
+		{
+			const int32 NumDimToCheck = FMath::Min(NumStackDimensionsA, NumStackDimensionsB);
+			for (int32 i = 0; i < NumDimToCheck; i++)
+			{
+				const uint32 VolumeA = InputA.Shape[NumStackDimensionsA - 1 - i];
+				const uint32 VolumeB = InputB.Shape[NumStackDimensionsB - 1 - i];
+
+				check(VolumeA == 1 || VolumeB == 1 || VolumeA == VolumeB);
+			}
+		}
+
+		const bool IsVectorA = InputA.Shape.Num() == 1;
+		const bool IsVectorB = InputB.Shape.Num() == 1;
+
+		Result.M = IsVectorA ? 1 : InputA.Shape[InputA.Shape.Num() - 2];
+		Result.N = IsVectorB ? 1 : InputB.Shape[InputB.Shape.Num() - 1];
+		Result.K = InputA.Shape[IsVectorA ? 0 : InputA.Shape.Num() - 1];
+		check(InputB.Shape[IsVectorB ? 0 : InputB.Shape.Num() - 2] == Result.K);
+
+		Result.StackShapeA.Init(1, NumStackDimensions);
+		Result.StackShapeB.Init(1, NumStackDimensions);
+
+		for (int32 i = 0; i < NumStackDimensions; i++)
+		{
+			int32 IdxA = InputA.Shape.Num() - 3 - i;
+			int32 IdxB = InputB.Shape.Num() - 3 - i;
+
+			Result.StackShapeA[Result.StackShapeA.Num() - 1 - i] = IdxA >= 0 ? InputA.Shape[IdxA] : 1;
+			Result.StackShapeB[Result.StackShapeB.Num() - 1 - i] = IdxB >= 0 ? InputB.Shape[IdxB] : 1;
+		}
+
+		Result.StackStrideA.Init(1, NumStackDimensions);
+		Result.StackStrideB.Init(1, NumStackDimensions);
+
+		for (int32 i = NumStackDimensions - 2; i >= 0; i--)
+		{
+			Result.StackStrideA[i] = Result.StackStrideA[i + 1] * Result.StackShapeA[i + 1];
+			Result.StackStrideB[i] = Result.StackStrideB[i + 1] * Result.StackShapeB[i + 1];
+		}
+
+		return Result;
+	}
+
+	void TGemmCS::ModifyCompilationEnvironment(const FGlobalShaderPermutationParameters& InParameters, FShaderCompilerEnvironment& OutEnvironment)
+	{
+		FGlobalShader::ModifyCompilationEnvironment(InParameters, OutEnvironment);
+		OutEnvironment.SetDefine(TEXT("MAX_NUM_STACK_DIMENSIONS"), FGemmConstants::MAX_NUM_STACK_DIMENSIONS);
+	}
+
+	void TGemmCS::FillInParameters(float Alpha, float Beta, int32 TransA, int32 TransB, const NNX::FMLTensorDesc& InputA, const NNX::FMLTensorDesc& InputB, const NNX::FMLTensorDesc& InputC, float CScalar, TGemmCS::FParameters& Parameters)
+	{
+		check(InputA.Shape.Num() == 2);
+		check(InputB.Shape.Num() == 2);
+
+		uint32 M = TransA != 0 ? InputA.Shape[1] : InputA.Shape[0];
+		uint32 K = TransA != 0 ? InputA.Shape[0] : InputA.Shape[1];
+		uint32 N = TransB != 0 ? InputB.Shape[0] : InputB.Shape[1];
+		check(K == (TransB != 0 ? InputB.Shape[1] : InputB.Shape[0]));
+
+		Parameters.Alpha = Alpha;
+		Parameters.Beta = Beta;
+		Parameters.TransA = TransA;
+		Parameters.TransB = TransB;
+		Parameters.M = M;
+		Parameters.N = N;
+		Parameters.K = K;
+		Parameters.MxK = M * K;
+		Parameters.KxN = K * N;
+		Parameters.MxN = M * N;
+		Parameters.CWidth = InputC.Shape.Num() == 0 ? 0 : InputC.Shape[InputC.Shape.Num() - 1];;
+		Parameters.CHeight = InputC.Shape.Num() == 0 ? 0 : InputC.Shape.Num() == 1 ? 1 : InputC.Shape[InputC.Shape.Num() - 2];
+		Parameters.CScalar = CScalar;
+	}
+
+	void TGemmCS::FillInParametersMatMul(const NNX::FMLTensorDesc& InputA, const NNX::FMLTensorDesc& InputB, TGemmCS::FParameters& Parameters)
+	{
+		GemmMatrixParameters Params = GetMatrixParametersMatMul(InputA, InputB);
+
+		Parameters.Alpha = 1.0f;
+		Parameters.Beta = 1.0f;
+		Parameters.TransA = 0;
+		Parameters.TransB = 0;
+		Parameters.M = Params.M;
+		Parameters.N = Params.N;
+		Parameters.K = Params.K;
+		Parameters.MxK = Params.M * Params.K;
+		Parameters.KxN = Params.K * Params.N;
+		Parameters.MxN = Params.M * Params.N;
+		Parameters.CWidth = 0;
+		Parameters.CHeight = 0;
+		Parameters.CScalar = 0.0f;
+
+		for (int32 i = 0; i < Params.StackShapeA.Num(); i++)
+		{
+			Parameters.StackShapeA_StackShapeB_StackStrideA_StackStrideB[i] = FUint32Vector4(Params.StackShapeA[i], Params.StackShapeB[i], Params.StackStrideA[i], Params.StackStrideB[i]);
+		}
+	}
+
+	FIntVector TGemmCS::GetGroupCount(const TGemmCS::FParameters& Parameters, EGemmAlgorithm Algorithm, int32 NumStackDimensions)
+	{
+		int32 OutputStackSize = 1;
+		for (int32 i = 0; i < NumStackDimensions; i++)
+		{
+			OutputStackSize *= FMath::Max(
+				Parameters.StackShapeA_StackShapeB_StackStrideA_StackStrideB[i].X,
+				Parameters.StackShapeA_StackShapeB_StackStrideA_StackStrideB[i].Y);
+		}
+
+		int32 ThreadGroupCountValueZ = FMath::DivideAndRoundUp(OutputStackSize, 1);
+		int32 ThreadGroupCountValueY = FMath::DivideAndRoundUp((int32)Parameters.M, 8);
+		int32 ThreadGroupCountValueX = FMath::DivideAndRoundUp((int32)Parameters.N, 8);
+		switch (Algorithm)
+		{
+		case EGemmAlgorithm::Simple8x8:
+		case EGemmAlgorithm::SharedMemory8x8:
+			ThreadGroupCountValueX = FMath::DivideAndRoundUp((int32)Parameters.N, 8);
+			ThreadGroupCountValueY = FMath::DivideAndRoundUp((int32)Parameters.M, 8);
+			break;
+		case EGemmAlgorithm::Simple16x16:
+		case EGemmAlgorithm::SharedMemory16x16:
+		case EGemmAlgorithm::MultiWrite1x16:
+		case EGemmAlgorithm::MultiWrite2x16:
+			ThreadGroupCountValueX = FMath::DivideAndRoundUp((int32)Parameters.N, 16);
+			ThreadGroupCountValueY = FMath::DivideAndRoundUp((int32)Parameters.M, 16);
+			break;
+		case EGemmAlgorithm::Simple32x32:
+		case EGemmAlgorithm::SharedMemory32x32:
+		case EGemmAlgorithm::MultiWrite1x32:
+		case EGemmAlgorithm::MultiWrite2x32:
+		case EGemmAlgorithm::MultiWrite4x32:
+			ThreadGroupCountValueX = FMath::DivideAndRoundUp((int32)Parameters.N, 32);
+			ThreadGroupCountValueY = FMath::DivideAndRoundUp((int32)Parameters.M, 32);
+			break;
+		case EGemmAlgorithm::Simple256x1:
+			ThreadGroupCountValueX = FMath::DivideAndRoundUp((int32)Parameters.N, 256);
+			ThreadGroupCountValueY = Parameters.M;
+			break;
+		case EGemmAlgorithm::MultiWrite2x64:
+		case EGemmAlgorithm::MultiWrite4x64:
+		case EGemmAlgorithm::MultiWrite8x64:
+			ThreadGroupCountValueX = FMath::DivideAndRoundUp((int32)Parameters.N, 64);
+			ThreadGroupCountValueY = FMath::DivideAndRoundUp((int32)Parameters.M, 64);
+			break;
+		default:
+			break;
+		}
+
+		return FIntVector(ThreadGroupCountValueX, ThreadGroupCountValueY, ThreadGroupCountValueZ);
+	}
+
+	EGemmAlgorithm TGemmCS::GetAlgorithm(const TGemmCS::FParameters& Parameters)
+	{
+		return EGemmAlgorithm::MultiWrite1x32;
+	}
+
+	IMPLEMENT_GLOBAL_SHADER(TGemmCS, "/NNX/GemmOp.usf", "main", SF_Compute);
+} // UE::NNIHlslShaders::Internal
\ No newline at end of file
diff --git a/Engine/Plugins/Experimental/NNX/Source/NNXHlslShaders/Private/NNXGatherCS.cpp b/Engine/Plugins/Experimental/NNX/Source/NNXHlslShaders/Private/NNXGatherCS.cpp
deleted file mode 100644
diff --git a/Engine/Plugins/Experimental/NNX/Source/NNXHlslShaders/Private/NNXGemmCS.cpp b/Engine/Plugins/Experimental/NNX/Source/NNXHlslShaders/Private/NNXGemmCS.cpp
deleted file mode 100644
diff --git a/Engine/Plugins/Experimental/NNX/Source/NNXHlslShaders/Public/NNXGatherCS.h b/Engine/Plugins/Experimental/NNX/Source/NNXHlslShaders/Public/NNXGatherCS.h
deleted file mode 100644
diff --git a/Engine/Plugins/Experimental/NNX/Source/NNXHlslShaders/Public/NNXGemmCS.h b/Engine/Plugins/Experimental/NNX/Source/NNXHlslShaders/Public/NNXGemmCS.h
deleted file mode 100644
diff --git a/Engine/Plugins/Experimental/NNX/Source/NNXRuntimeRDG/Private/Hlsl/NNIRuntimeRDGElementWiseBinary.cpp b/Engine/Plugins/Experimental/NNX/Source/NNXRuntimeRDG/Private/Hlsl/NNIRuntimeRDGElementWiseBinary.cpp
--- a/Engine/Plugins/Experimental/NNX/Source/NNXRuntimeRDG/Private/Hlsl/NNIRuntimeRDGElementWiseBinary.cpp
+++ b/Engine/Plugins/Experimental/NNX/Source/NNXRuntimeRDG/Private/Hlsl/NNIRuntimeRDGElementWiseBinary.cpp
@@ -15,18 +15,18 @@ namespace UE::NNIRuntimeRDG::Private::Hlsl
 	 * Binary element-wise operator implementation
 	 */
 	template<EMLElementWiseBinaryOperatorType OpType>
-	class TOperatorElementWiseBinary : public NNX::FMLOperatorHlsl
+	class TElementWiseBinary : public NNX::FMLOperatorHlsl
 	{
 	public:
 
-		TOperatorElementWiseBinary() {}
-		virtual ~TOperatorElementWiseBinary() = default;
+		TElementWiseBinary() {}
+		virtual ~TElementWiseBinary() = default;
 
 	private:
 
-		NNX::FMLTensorDesc LHSInput;
-		NNX::FMLTensorDesc RHSInput;
-		NNX::FMLTensorDesc Output;
+		NNX::FMLTensorDesc LHSInput = {};
+		NNX::FMLTensorDesc RHSInput = {};
+		NNX::FMLTensorDesc Output = {};
 
 	public:
 
@@ -85,7 +85,7 @@ namespace UE::NNIRuntimeRDG::Private::Hlsl
 	template<EMLElementWiseBinaryOperatorType OpType>
 	NNX::FMLOperatorHlsl* CreateElementWiseBinaryOperator()
 	{
-		return new TOperatorElementWiseBinary<OpType>();
+		return new TElementWiseBinary<OpType>();
 	}
 
 	bool RegisterElementWiseBinaryOperators(NNX::FMLOperatorRegistryHlsl& Registry)
diff --git a/Engine/Plugins/Experimental/NNX/Source/NNXRuntimeRDG/Private/Hlsl/NNIRuntimeRDGElementWiseUnary.cpp b/Engine/Plugins/Experimental/NNX/Source/NNXRuntimeRDG/Private/Hlsl/NNIRuntimeRDGElementWiseUnary.cpp
--- a/Engine/Plugins/Experimental/NNX/Source/NNXRuntimeRDG/Private/Hlsl/NNIRuntimeRDGElementWiseUnary.cpp
+++ b/Engine/Plugins/Experimental/NNX/Source/NNXRuntimeRDG/Private/Hlsl/NNIRuntimeRDGElementWiseUnary.cpp
@@ -15,20 +15,20 @@ namespace UE::NNIRuntimeRDG::Private::Hlsl
 	 * Unary element-wise operator implementation
 	 */
 	template<EMLElementWiseUnaryOperatorType OpType>
-	class TOperatorElementWiseUnary : public NNX::FMLOperatorHlsl
+	class TElementWiseUnary : public NNX::FMLOperatorHlsl
 	{
 	public:
 
-		TOperatorElementWiseUnary() : Alpha(0.0f), Beta(0.0f), Gamma(0.0f) {}
-		virtual ~TOperatorElementWiseUnary() = default;
+		TElementWiseUnary() {}
+		virtual ~TElementWiseUnary() = default;
 
 	private:
 
-		float Alpha;
-		float Beta;
-		float Gamma;
-		NNX::FMLTensorDesc Input;
-		NNX::FMLTensorDesc Output;
+		float Alpha = 0.0f;
+		float Beta = 0.0f;
+		float Gamma = 0.0f;
+		NNX::FMLTensorDesc Input = {};
+		NNX::FMLTensorDesc Output = {};
 
 	public:
 
@@ -84,22 +84,22 @@ namespace UE::NNIRuntimeRDG::Private::Hlsl
 		}
 	};
 
-	template<> TOperatorElementWiseUnary<EMLElementWiseUnaryOperatorType::Selu>::TOperatorElementWiseUnary()
+	template<> TElementWiseUnary<EMLElementWiseUnaryOperatorType::Selu>::TElementWiseUnary()
 		: Alpha(1.67326319217681884765625f), Beta(0.0f), Gamma(1.05070102214813232421875f)
 	{
 	}
 
-	template<> TOperatorElementWiseUnary<EMLElementWiseUnaryOperatorType::Elu>::TOperatorElementWiseUnary()
+	template<> TElementWiseUnary<EMLElementWiseUnaryOperatorType::Elu>::TElementWiseUnary()
 		: Alpha(1.0f), Beta(0.0f), Gamma(0.0f) 
 	{
 	}
 
-	template<> TOperatorElementWiseUnary<EMLElementWiseUnaryOperatorType::HardSigmoid>::TOperatorElementWiseUnary()
+	template<> TElementWiseUnary<EMLElementWiseUnaryOperatorType::HardSigmoid>::TElementWiseUnary()
 		: Alpha(0.2f), Beta(0.5f), Gamma(0.0f)
 	{
 	}
 
-	template<> TOperatorElementWiseUnary<EMLElementWiseUnaryOperatorType::LeakyRelu>::TOperatorElementWiseUnary()
+	template<> TElementWiseUnary<EMLElementWiseUnaryOperatorType::LeakyRelu>::TElementWiseUnary()
 		: Alpha(0.01f), Beta(0.0f), Gamma(0.0f)
 	{
 	}
@@ -107,7 +107,7 @@ namespace UE::NNIRuntimeRDG::Private::Hlsl
 	template<EMLElementWiseUnaryOperatorType OpType>
 	NNX::FMLOperatorHlsl* CreateElementWiseUnaryOperator()
 	{
-		return new TOperatorElementWiseUnary<OpType>();
+		return new TElementWiseUnary<OpType>();
 	}
 
 	bool RegisterElementWiseUnaryOperators(NNX::FMLOperatorRegistryHlsl& Registry)
diff --git a/Engine/Plugins/Experimental/NNX/Source/NNXRuntimeRDG/Private/Hlsl/NNIRuntimeRDGElementWiseVariadic.cpp b/Engine/Plugins/Experimental/NNX/Source/NNXRuntimeRDG/Private/Hlsl/NNIRuntimeRDGElementWiseVariadic.cpp
--- a/Engine/Plugins/Experimental/NNX/Source/NNXRuntimeRDG/Private/Hlsl/NNIRuntimeRDGElementWiseVariadic.cpp
+++ b/Engine/Plugins/Experimental/NNX/Source/NNXRuntimeRDG/Private/Hlsl/NNIRuntimeRDGElementWiseVariadic.cpp
@@ -85,17 +85,17 @@ namespace UE::NNIRuntimeRDG::Private::Hlsl
 	 * Variadic Element-wise ML operator implementation
 	 */
 	template<EMLElementWiseVariadicOperatorType OpType>
-	class TOperatorElementWiseVariadic : public NNX::FMLOperatorHlsl
+	class TElementWiseVariadic : public NNX::FMLOperatorHlsl
 	{
 	public:
 
-		TOperatorElementWiseVariadic() {}
-		virtual ~TOperatorElementWiseVariadic() = default;
+		TElementWiseVariadic() {}
+		virtual ~TElementWiseVariadic() = default;
 
 	private:
 
 		TArray<NNX::FMLTensorDesc> InputDescs;
-		NNX::FMLTensorDesc OutputDesc;
+		NNX::FMLTensorDesc OutputDesc = {};
 
 	public:
 
@@ -153,7 +153,7 @@ namespace UE::NNIRuntimeRDG::Private::Hlsl
 	template<EMLElementWiseVariadicOperatorType OpType>
 	NNX::FMLOperatorHlsl* CreateElementWiseVariadicOperator()
 	{
-		return new TOperatorElementWiseVariadic<OpType>();
+		return new TElementWiseVariadic<OpType>();
 	}
 
 	bool RegisterElementWiseVariadicOperators(NNX::FMLOperatorRegistryHlsl& Registry)
diff --git a/Engine/Plugins/Experimental/NNX/Source/NNXRuntimeRDG/Private/Hlsl/NNIRuntimeRDGGather.cpp b/Engine/Plugins/Experimental/NNX/Source/NNXRuntimeRDG/Private/Hlsl/NNIRuntimeRDGGather.cpp
new file mode 100644
--- /dev/null
+++ b/Engine/Plugins/Experimental/NNX/Source/NNXRuntimeRDG/Private/Hlsl/NNIRuntimeRDGGather.cpp
@@ -0,0 +1,94 @@
+// Copyright Epic Games, Inc. All Rights Reserved.
+
+#include "NNIRuntimeRDGGather.h"
+#include "NNXRuntimeHLSLHelper.h"
+#include "NNIHlslShadersGatherCS.h"
+
+namespace UE::NNIRuntimeRDG::Private::Hlsl
+{
+	DECLARE_GPU_STAT_NAMED(FNNIOperatorGather, TEXT("NNI.Operator.Hlsl.Gather"));
+
+	/**
+	 * Gather operator implementation
+	 */
+	template <typename DataElementType, typename IndicesElementType>
+	class FGather : public NNX::FMLOperatorHlsl
+	{
+	public:
+
+		FGather() {}
+		virtual ~FGather() = default;
+
+	private:
+
+		int32 Axis = 0;
+		NNX::FMLTensorDesc Data = {};
+		NNX::FMLTensorDesc Indices = {};
+		NNX::FMLTensorDesc Output = {};
+
+	public:
+
+		virtual bool Initialize(TArrayView<const NNX::FMLTensorDesc> InputTensors, TArrayView<const NNX::FMLTensorDesc> OutputTensors, const FMLAttributeMap& Attributes) override
+		{
+			using namespace UE::NNIHlslShaders::Internal;
+
+			check(InputTensors.Num() == 2)
+			check(OutputTensors.Num() == 1)
+			check(Output.Shape.Num() <= FGatherConstants::MAX_NUM_DIMENSIONS)
+			check(InputTensors[0].Shape.Num() > 0)
+			check(InputTensors[1].Shape.Num() > 0)
+			check(InputTensors[0].Shape.Num() + (InputTensors[1].Shape.Num() - 1) <= FGatherConstants::MAX_NUM_DIMENSIONS)
+
+			Axis = Attributes.GetOptionalFloat(TEXT("axis"), Axis);
+			check(Axis < InputTensors[0].Shape.Num())
+			check(Axis >= -InputTensors[0].Shape.Num())
+			Axis = Axis >= 0 ? Axis : InputTensors[0].Shape.Num() + Axis;
+
+			Data = InputTensors[0];
+			Indices = InputTensors[1];
+			Output = OutputTensors[0];
+
+			return true;
+		}
+
+		virtual void Dispatch(FRDGBuilder& GraphBuilder, TArrayView<const NNX::FMLTensorBinding> InInputBindings, TArrayView<const NNX::FMLTensorBinding> OutOutputBindings) override
+		{
+			using namespace UE::NNIHlslShaders::Internal;
+
+			// Set parameters
+			TGatherCS::FParameters* Parameters = GraphBuilder.AllocParameters<TGatherCS::FParameters>();
+			TGatherCS::FillInParameters(Axis, Data, Indices, *Parameters);
+			Parameters->Data = GraphBuilder.CreateSRV(FRDGBufferSRVDesc(InInputBindings[0].Buffer, PF_R32_FLOAT));
+			Parameters->Indices = GraphBuilder.CreateSRV(FRDGBufferSRVDesc(InInputBindings[1].Buffer, PF_R32_FLOAT));
+			Parameters->Output = GraphBuilder.CreateUAV(FRDGBufferUAVDesc(OutOutputBindings[0].Buffer, PF_R32_FLOAT));
+
+			TGatherCS::FPermutationDomain PermutationVector;
+			PermutationVector.Set<TGatherCS::FGatherNumOutputDimensions>(Output.Shape.Num());
+			TShaderMapRef<TGatherCS> ComputeShader(GetGlobalShaderMap(GMaxRHIFeatureLevel), PermutationVector);
+
+			FIntVector ThreadGroupCount = TGatherCS::GetGroupCount(*Parameters);
+
+			RDG_EVENT_SCOPE(GraphBuilder, "NNI.Operator.Hlsl.Gather");
+			RDG_GPU_STAT_SCOPE(GraphBuilder, FNNIOperatorGather);
+
+			FComputeShaderUtils::AddPass(
+				GraphBuilder,
+				RDG_EVENT_NAME("NNI.Operator.Hlsl.Gather.Dispatch"),
+				ERDGPassFlags::Compute | ERDGPassFlags::NeverCull,
+				ComputeShader,
+				Parameters,
+				ThreadGroupCount);
+		}
+	};
+
+	NNX::FMLOperatorHlsl* CreateGatherOperator()
+	{
+		return new FGather<float, int32>();
+	}
+
+	bool RegisterGatherOperator(NNX::FMLOperatorRegistryHlsl& Registry)
+	{
+		Registry.OpAdd(TEXT("Gather"), CreateGatherOperator);
+		return true;
+	}
+} // UE::NNIRuntimeRDG::Private::Hlsl
\ No newline at end of file
diff --git a/Engine/Plugins/Experimental/NNX/Source/NNXRuntimeRDG/Private/Hlsl/NNIRuntimeRDGGather.h b/Engine/Plugins/Experimental/NNX/Source/NNXRuntimeRDG/Private/Hlsl/NNIRuntimeRDGGather.h
new file mode 100644
--- /dev/null
+++ b/Engine/Plugins/Experimental/NNX/Source/NNXRuntimeRDG/Private/Hlsl/NNIRuntimeRDGGather.h
@@ -0,0 +1,11 @@
+// Copyright Epic Games, Inc. All Rights Reserved.
+
+#pragma once
+
+#include "../NNXRuntimeRDG.h"
+#include "../NNXRuntimeHLSLOp.h"
+
+namespace UE::NNIRuntimeRDG::Private::Hlsl
+{
+	bool RegisterGatherOperator(NNX::FMLOperatorRegistryHlsl& Registry);
+} // UE::NNIRuntimeRDG::Private::Hlsl
diff --git a/Engine/Plugins/Experimental/NNX/Source/NNXRuntimeRDG/Private/Hlsl/NNIRuntimeRDGGemm.cpp b/Engine/Plugins/Experimental/NNX/Source/NNXRuntimeRDG/Private/Hlsl/NNIRuntimeRDGGemm.cpp
new file mode 100644
--- /dev/null
+++ b/Engine/Plugins/Experimental/NNX/Source/NNXRuntimeRDG/Private/Hlsl/NNIRuntimeRDGGemm.cpp
@@ -0,0 +1,142 @@
+// Copyright Epic Games, Inc. All Rights Reserved.
+
+#include "NNIRuntimeRDGGemm.h"
+#include "NNIHlslShadersGemmCS.h"
+#include "NNXRuntimeHLSLHelper.h"
+
+namespace UE::NNIRuntimeRDG::Private::Hlsl
+{
+	DECLARE_GPU_STAT_NAMED(FNNIOperatorGemm, TEXT("NNI.Operator.Hlsl.Gemm"));
+
+	/**
+	 * Gemm operator implementation
+	 */
+	class TGemm : public NNX::FMLOperatorHlsl
+	{
+	public:
+
+		TGemm() {}
+		virtual ~TGemm() = default;
+
+	private:
+
+		NNX::FMLTensorDesc InputA = {};
+		NNX::FMLTensorDesc InputB = {};
+		NNX::FMLTensorDesc InputC = {};
+		NNX::FMLTensorDesc Output = {};
+
+		float InputAlpha = 1.0f;
+		float InputBeta = 1.0f;
+		int32 InputTransA = 0;
+		int32 InputTransB = 0;
+		uint32 InputM = 0;
+		uint32 InputN = 0;
+		uint32 InputK = 0;
+		uint32 InputCWidth = 0;
+		uint32 InputCHeight = 0;
+
+		bool bIsCScalar = false;
+		bool bNoBias = true;
+
+	public:
+
+		virtual bool Initialize(TArrayView<const NNX::FMLTensorDesc> InputTensors, TArrayView<const NNX::FMLTensorDesc> OutputTensors, const FMLAttributeMap& Attributes) override
+		{
+			check(InputTensors.Num() >= 2 && InputTensors.Num() <= 3);
+			check(OutputTensors.Num() == 1);
+
+			InputA = InputTensors[0];
+			InputB = InputTensors[1];
+			if (InputTensors.Num() == 3) InputC = InputTensors[2];
+			Output = OutputTensors[0];
+
+			check(InputA.Shape.Num() == 2);
+			check(InputB.Shape.Num() == 2);
+			check(InputC.Shape.Num() < 3);
+
+			check(InputC.Shape.Num() != 1 || InputC.Shape[0] != 1); // TODO scalar version not supported yet
+
+			// C is treated as a scalar if there is no valid C, either width or height is zero or C dimension is 1x1
+			bIsCScalar = false; // InputTensors.Num() != 3 || InputC.Sizes[0] * InputC.Sizes[1] < 2;
+			// CScalar = C != nullptr ? C[0] : (InElementType)0;
+			bNoBias = InputTensors.Num() != 3 /*|| InputC.Sizes[0] * InputC.Sizes[1] < 1*/;
+
+			InputAlpha = Attributes.GetOptionalFloat(TEXT("alpha"), InputAlpha);
+			InputBeta = Attributes.GetOptionalFloat(TEXT("beta"), InputBeta);
+			InputTransA = Attributes.GetOptionalInt32(TEXT("transA"), InputTransA);
+			InputTransB = Attributes.GetOptionalInt32(TEXT("transB"), InputTransB);
+
+			return true;
+		}
+
+		virtual void Dispatch(FRDGBuilder& GraphBuilder, TArrayView<const NNX::FMLTensorBinding> InInputBindings, TArrayView<const NNX::FMLTensorBinding> OutOutputBindings) override
+		{
+			using namespace UE::NNIHlslShaders::Internal;
+
+			const EGemmAlgorithm Algorithm = EGemmAlgorithm::Simple32x32;
+
+			const float CScalar = 0.0f;
+
+			// Set parameters
+			TGemmCS::FParameters* Parameters = GraphBuilder.AllocParameters<TGemmCS::FParameters>();
+			TGemmCS::FillInParameters(InputAlpha, InputBeta, InputTransA, InputTransB, InputA, InputB, InputC, CScalar, *Parameters);
+			Parameters->A = GraphBuilder.CreateSRV(FRDGBufferSRVDesc(InInputBindings[0].Buffer, PF_R32_FLOAT));
+			Parameters->B = GraphBuilder.CreateSRV(FRDGBufferSRVDesc(InInputBindings[1].Buffer, PF_R32_FLOAT));
+			if (InInputBindings.Num() == 3) {
+				Parameters->C = GraphBuilder.CreateSRV(FRDGBufferSRVDesc(InInputBindings[2].Buffer, PF_R32_FLOAT));
+			}
+			Parameters->Y = GraphBuilder.CreateUAV(FRDGBufferUAVDesc(OutOutputBindings[0].Buffer, PF_R32_FLOAT));
+
+			TGemmCS::FPermutationDomain PermutationVector;
+			PermutationVector.Set<TGemmCS::FGemmCScalar>(bNoBias ? EGemmCScalar::NoBias : (bIsCScalar ? EGemmCScalar::Yes : EGemmCScalar::No));
+			PermutationVector.Set<TGemmCS::FGemmAlgorithm>(Algorithm);
+			PermutationVector.Set<TGemmCS::FGemmNumStackDimensions>(0);
+			TShaderMapRef<TGemmCS> ComputeShader(GetGlobalShaderMap(GMaxRHIFeatureLevel), PermutationVector);
+
+			FIntVector ThreadGroupCount = TGemmCS::GetGroupCount(*Parameters, Algorithm, 0);
+
+			RDG_EVENT_SCOPE(GraphBuilder, "NNI.Operator.Hlsl.Gemm");
+			RDG_GPU_STAT_SCOPE(GraphBuilder, FNNIOperatorGemm);
+
+			FComputeShaderUtils::AddPass(
+				GraphBuilder,
+				RDG_EVENT_NAME("NNI.Operator.Hlsl.Dispatch"),
+				ERDGPassFlags::Compute | ERDGPassFlags::NeverCull,
+				ComputeShader,
+				Parameters,
+				ThreadGroupCount);
+		}
+	};
+
+	bool ValidateGemmOperator(const FMLAttributeMap& AttributeMap, TConstArrayView<EMLTensorDataType> InputTypes)
+	{
+		bool bIsValid = true;
+
+		NNX::FAttributeValidator AttributeValidator;
+		AttributeValidator.AddOptional(TEXT("alpha"), EMLAttributeDataType::Float);
+		AttributeValidator.AddOptional(TEXT("beta"), EMLAttributeDataType::Float);
+		AttributeValidator.AddOptional(TEXT("transA"), EMLAttributeDataType::Int32);
+		AttributeValidator.AddOptional(TEXT("transB"), EMLAttributeDataType::Int32);
+		bIsValid &= AttributeValidator.Validate(AttributeMap);
+
+		NNX::FInputValidator InputValidator;
+		InputValidator.AddSupportedType(EMLTensorDataType::Float);
+		InputValidator.AddRequired();
+		InputValidator.AddRequired();
+		InputValidator.AddOptional();
+		bIsValid &= InputValidator.Validate(InputTypes);
+
+		return bIsValid;
+	}
+
+	NNX::FMLOperatorHlsl* CreateGemmOperator()
+	{
+		return new TGemm();
+	}
+
+	bool RegisterGemmOperator(NNX::FMLOperatorRegistryHlsl& Registry)
+	{
+		Registry.OpAdd(TEXT("Gemm"), CreateGemmOperator, ValidateGemmOperator);
+		return true;
+	}
+} // UE::NNIRuntimeRDG::Private::Hlsl
diff --git a/Engine/Plugins/Experimental/NNX/Source/NNXRuntimeRDG/Private/Hlsl/NNIRuntimeRDGGemm.h b/Engine/Plugins/Experimental/NNX/Source/NNXRuntimeRDG/Private/Hlsl/NNIRuntimeRDGGemm.h
new file mode 100644
--- /dev/null
+++ b/Engine/Plugins/Experimental/NNX/Source/NNXRuntimeRDG/Private/Hlsl/NNIRuntimeRDGGemm.h
@@ -0,0 +1,11 @@
+// Copyright Epic Games, Inc. All Rights Reserved.
+
+#pragma once
+
+#include "../NNXRuntimeRDG.h"
+#include "../NNXRuntimeHLSLOp.h"
+
+namespace UE::NNIRuntimeRDG::Private::Hlsl
+{
+	bool RegisterGemmOperator(NNX::FMLOperatorRegistryHlsl& Registry);
+} // UE::NNIRuntimeRDG::Private::Hlsl
diff --git a/Engine/Plugins/Experimental/NNX/Source/NNXRuntimeRDG/Private/Hlsl/NNIRuntimeRDGMatMul.cpp b/Engine/Plugins/Experimental/NNX/Source/NNXRuntimeRDG/Private/Hlsl/NNIRuntimeRDGMatMul.cpp
new file mode 100644
--- /dev/null
+++ b/Engine/Plugins/Experimental/NNX/Source/NNXRuntimeRDG/Private/Hlsl/NNIRuntimeRDGMatMul.cpp
@@ -0,0 +1,87 @@
+// Copyright Epic Games, Inc. All Rights Reserved.
+
+#include "NNIRuntimeRDGMatMul.h"
+#include "NNIHlslShadersGemmCS.h"
+#include "NNXRuntimeHLSLHelper.h"
+
+namespace UE::NNIRuntimeRDG::Private::Hlsl
+{
+    DECLARE_GPU_STAT_NAMED(FNNIOperatorMatMul, TEXT("NNI.Operator.Hlsl.MatMul"));
+
+	/**
+	 * MatMul operator implementation
+	 */
+	class TMatMul : public NNX::FMLOperatorHlsl
+	{
+	public:
+
+		TMatMul() {}
+		virtual ~TMatMul() = default;
+
+	private:
+
+		NNX::FMLTensorDesc InputA = {};
+		NNX::FMLTensorDesc InputB = {};
+		NNX::FMLTensorDesc Output = {};
+
+	public:
+
+		virtual bool Initialize(TArrayView<const NNX::FMLTensorDesc> InputTensors, TArrayView<const NNX::FMLTensorDesc> OutputTensors, const FMLAttributeMap& Attributes) override
+		{
+			check(InputTensors.Num() == 2);
+			check(OutputTensors.Num() == 1);
+
+			InputA = InputTensors[0];
+			InputB = InputTensors[1];
+			Output = OutputTensors[0];
+
+			return true;
+		}
+
+		virtual void Dispatch(FRDGBuilder& GraphBuilder, TArrayView<const NNX::FMLTensorBinding> InInputBindings, TArrayView<const NNX::FMLTensorBinding> OutOutputBindings) override
+		{
+			using namespace UE::NNIHlslShaders::Internal;
+
+			const EGemmAlgorithm Algorithm = EGemmAlgorithm::Simple32x32;
+
+			const int32 NumStackDimensions = FMath::Max(FMath::Max(InputA.Shape.Num(), InputB.Shape.Num()) - 2, 0);
+
+			// Set parameters
+			TGemmCS::FParameters* Parameters = GraphBuilder.AllocParameters<TGemmCS::FParameters>();
+			TGemmCS::FillInParametersMatMul(InputA, InputB, *Parameters);
+			Parameters->A = GraphBuilder.CreateSRV(FRDGBufferSRVDesc(InInputBindings[0].Buffer, PF_R32_FLOAT));
+			Parameters->B = GraphBuilder.CreateSRV(FRDGBufferSRVDesc(InInputBindings[1].Buffer, PF_R32_FLOAT));
+			Parameters->Y = GraphBuilder.CreateUAV(FRDGBufferUAVDesc(OutOutputBindings[0].Buffer, PF_R32_FLOAT));
+
+			TGemmCS::FPermutationDomain PermutationVector;
+			PermutationVector.Set<TGemmCS::FGemmCScalar>(EGemmCScalar::NoBias);
+			PermutationVector.Set<TGemmCS::FGemmAlgorithm>(Algorithm);
+			PermutationVector.Set<TGemmCS::FGemmNumStackDimensions>(NumStackDimensions);
+			TShaderMapRef<TGemmCS> ComputeShader(GetGlobalShaderMap(GMaxRHIFeatureLevel), PermutationVector);
+
+			FIntVector ThreadGroupCount = TGemmCS::GetGroupCount(*Parameters, Algorithm, NumStackDimensions);
+
+			RDG_EVENT_SCOPE(GraphBuilder, "NNI.Operator.Hlsl.MatMul");
+			RDG_GPU_STAT_SCOPE(GraphBuilder, FNNIOperatorMatMul);
+
+			FComputeShaderUtils::AddPass(
+				GraphBuilder,
+				RDG_EVENT_NAME("NNI.Operator.Hlsl.MatMul.Dispatch"),
+				ERDGPassFlags::Compute | ERDGPassFlags::NeverCull,
+				ComputeShader,
+				Parameters,
+				ThreadGroupCount);
+		}
+	};
+
+	NNX::FMLOperatorHlsl* CreateMatMulOperator()
+	{
+		return new TMatMul();
+	}
+
+	bool RegisterMatMulOperator(NNX::FMLOperatorRegistryHlsl& Registry)
+	{
+		Registry.OpAdd(TEXT("MatMul"), CreateMatMulOperator);
+		return true;
+	}
+} // UE::NNIRuntimeRDG::Private::Hlsl
diff --git a/Engine/Plugins/Experimental/NNX/Source/NNXRuntimeRDG/Private/Hlsl/NNIRuntimeRDGMatMul.h b/Engine/Plugins/Experimental/NNX/Source/NNXRuntimeRDG/Private/Hlsl/NNIRuntimeRDGMatMul.h
new file mode 100644
--- /dev/null
+++ b/Engine/Plugins/Experimental/NNX/Source/NNXRuntimeRDG/Private/Hlsl/NNIRuntimeRDGMatMul.h
@@ -0,0 +1,11 @@
+// Copyright Epic Games, Inc. All Rights Reserved.
+
+#pragma once
+
+#include "../NNXRuntimeRDG.h"
+#include "../NNXRuntimeHLSLOp.h"
+
+namespace UE::NNIRuntimeRDG::Private::Hlsl
+{
+	bool RegisterMatMulOperator(NNX::FMLOperatorRegistryHlsl& Registry);
+} // UE::NNIRuntimeRDG::Private::Hlsl
diff --git a/Engine/Plugins/Experimental/NNX/Source/NNXRuntimeRDG/Private/NNXRuntimeHLSL.cpp b/Engine/Plugins/Experimental/NNX/Source/NNXRuntimeRDG/Private/NNXRuntimeHLSL.cpp
--- a/Engine/Plugins/Experimental/NNX/Source/NNXRuntimeRDG/Private/NNXRuntimeHLSL.cpp
+++ b/Engine/Plugins/Experimental/NNX/Source/NNXRuntimeRDG/Private/NNXRuntimeHLSL.cpp
@@ -1,7 +1,5 @@
 // Copyright Epic Games, Inc. All Rights Reserved.
 
-#include "Hlsl/NNIRuntimeRDGConv.h"
-#include "Hlsl/NNIRuntimeRDGConvTranspose.h"
 #include "NNXInferenceModel.h"
 #include "NNXRuntimeFormat.h"
 #include "NNXModelOptimizer.h"
@@ -9,8 +7,10 @@
 #include "Hlsl/NNIRuntimeRDGElementWiseBinary.h"
 #include "Hlsl/NNIRuntimeRDGElementWiseUnary.h"
 #include "Hlsl/NNIRuntimeRDGElementWiseVariadic.h"
-#include "NNXRuntimeHLSLGemmOp.h"
-#include "NNXRuntimeHLSLMatMulOp.h"
+#include "Hlsl/NNIRuntimeRDGConv.h"
+#include "Hlsl/NNIRuntimeRDGConvTranspose.h"
+#include "Hlsl/NNIRuntimeRDGGemm.h"
+#include "Hlsl/NNIRuntimeRDGMatMul.h"
 #include "NNXRuntimeHLSLHelper.h"
 #include "NNXRuntimeRDG.h"
 
@@ -164,10 +164,10 @@ public:
 		UE::NNIRuntimeRDG::Private::Hlsl::RegisterElementWiseUnaryOperators(*registry);
 		UE::NNIRuntimeRDG::Private::Hlsl::RegisterElementWiseBinaryOperators(*registry);
 		UE::NNIRuntimeRDG::Private::Hlsl::RegisterElementWiseVariadicOperators(*registry);
-		RegisterGemmOperator(*registry);
+		UE::NNIRuntimeRDG::Private::Hlsl::RegisterGemmOperator(*registry);
 		UE::NNIRuntimeRDG::Private::Hlsl::RegisterConvOperator(*registry);
 		UE::NNIRuntimeRDG::Private::Hlsl::RegisterConvTransposeOperator(*registry);
-		RegisterMatMulOperator(*registry);
+		UE::NNIRuntimeRDG::Private::Hlsl::RegisterMatMulOperator(*registry);
 
 		return true;
 	}
diff --git a/Engine/Plugins/Experimental/NNX/Source/NNXRuntimeRDG/Private/NNXRuntimeHLSLGatherOp.cpp b/Engine/Plugins/Experimental/NNX/Source/NNXRuntimeRDG/Private/NNXRuntimeHLSLGatherOp.cpp
deleted file mode 100644
diff --git a/Engine/Plugins/Experimental/NNX/Source/NNXRuntimeRDG/Private/NNXRuntimeHLSLGatherOp.h b/Engine/Plugins/Experimental/NNX/Source/NNXRuntimeRDG/Private/NNXRuntimeHLSLGatherOp.h
deleted file mode 100644
diff --git a/Engine/Plugins/Experimental/NNX/Source/NNXRuntimeRDG/Private/NNXRuntimeHLSLGemmOp.cpp b/Engine/Plugins/Experimental/NNX/Source/NNXRuntimeRDG/Private/NNXRuntimeHLSLGemmOp.cpp
deleted file mode 100644
diff --git a/Engine/Plugins/Experimental/NNX/Source/NNXRuntimeRDG/Private/NNXRuntimeHLSLGemmOp.h b/Engine/Plugins/Experimental/NNX/Source/NNXRuntimeRDG/Private/NNXRuntimeHLSLGemmOp.h
deleted file mode 100644
diff --git a/Engine/Plugins/Experimental/NNX/Source/NNXRuntimeRDG/Private/NNXRuntimeHLSLMatMulOp.cpp b/Engine/Plugins/Experimental/NNX/Source/NNXRuntimeRDG/Private/NNXRuntimeHLSLMatMulOp.cpp
deleted file mode 100644
diff --git a/Engine/Plugins/Experimental/NNX/Source/NNXRuntimeRDG/Private/NNXRuntimeHLSLMatMulOp.h b/Engine/Plugins/Experimental/NNX/Source/NNXRuntimeRDG/Private/NNXRuntimeHLSLMatMulOp.h
deleted file mode 100644
