NNI - Cleaning NNIHlslShaders module

diff --git a/Engine/Plugins/Experimental/NNX/NNX.uplugin b/Engine/Plugins/Experimental/NNX/NNX.uplugin
--- a/Engine/Plugins/Experimental/NNX/NNX.uplugin
+++ b/Engine/Plugins/Experimental/NNX/NNX.uplugin
@@ -26,7 +26,7 @@
 			"PlatformAllowList": [ "Win64", "Linux", "Mac" ]
 		},
 		{
-			"Name": "NNXHlslShaders",
+			"Name": "NNIHlslShaders",
 			"Type": "RuntimeAndProgram",
 			"LoadingPhase": "PostConfigInit",
 			"PlatformAllowList": [ "Win64", "Linux", "Mac" ]
diff --git a/Engine/Plugins/Experimental/NNX/Source/NNIHlslShaders/Internal/NNIHlslShadersConvCS.h b/Engine/Plugins/Experimental/NNX/Source/NNIHlslShaders/Internal/NNIHlslShadersConvCS.h
new file mode 100644
--- /dev/null
+++ b/Engine/Plugins/Experimental/NNX/Source/NNIHlslShaders/Internal/NNIHlslShadersConvCS.h
@@ -0,0 +1,109 @@
+// Copyright Epic Games, Inc. All Rights Reserved.
+
+#pragma once
+
+#include "GlobalShader.h"
+#include "ShaderParameterUtils.h"
+#include "RenderGraphUtils.h"
+
+namespace UE::NNIHlslShaders::Internal
+{
+	enum class EConvAlgorithm : uint8
+	{
+		SharedMemory = 0,
+		MAX
+	};
+
+	enum class EConvGroupSize : uint8
+	{
+		Size128 = 0,
+		Size256,
+		Size512,
+		MAX
+	};
+
+	enum class EConvAutoPad : uint8
+	{
+		NOTSET = 0,// Use pad values passed in the array
+		SAME_UPPER,// Auto-pad to match input and output shape with potetnial extra padding at the end
+		SAME_LOWER,// Auto-pad to match input and output shape with potetnial extra padding at the beginning
+		VALID,// Set all paddings to zero
+		MAX
+	};
+
+	class FConvConstants
+	{
+	public:
+		static const int32 MAX_NUM_DIMENSIONS{4};
+		static const int32 MIN_NUM_READS_PER_THREAD_POW2{1};
+		static const int32 MAX_NUM_READS_PER_THREAD_POW2{3};
+	};
+
+	class NNIHLSLSHADERS_API FConvCS : public FGlobalShader
+	{
+		DECLARE_GLOBAL_SHADER(FConvCS);
+		SHADER_USE_PARAMETER_STRUCT(FConvCS, FGlobalShader)
+
+		class FConvAlgorithm : SHADER_PERMUTATION_ENUM_CLASS("ALGORITHM", EConvAlgorithm);
+		class FConvGroupSize : SHADER_PERMUTATION_ENUM_CLASS("GROUP_SIZE", EConvGroupSize);
+		class FConvNumDimensions : SHADER_PERMUTATION_RANGE_INT("NUM_DIMENSIONS", 1, FConvConstants::MAX_NUM_DIMENSIONS);
+		class FConvNumReadsPerThread : SHADER_PERMUTATION_RANGE_INT("NUM_READS_PER_THREAD_POW2", FConvConstants::MIN_NUM_READS_PER_THREAD_POW2, FConvConstants::MAX_NUM_READS_PER_THREAD_POW2);
+		class FConvHasB : SHADER_PERMUTATION_BOOL("HAS_B");
+		using FPermutationDomain = TShaderPermutationDomain<FConvAlgorithm, FConvGroupSize, FConvNumDimensions, FConvNumReadsPerThread, FConvHasB>;
+
+	public:
+		BEGIN_SHADER_PARAMETER_STRUCT(FParameters, )
+			SHADER_PARAMETER_RDG_BUFFER_SRV(Buffer<float>, X)
+			SHADER_PARAMETER_RDG_BUFFER_SRV(Buffer<float>, W)
+			SHADER_PARAMETER_RDG_BUFFER_UAV(RWBuffer<float>, Y)
+			SHADER_PARAMETER_RDG_BUFFER_SRV(Buffer<float>, B)
+			SHADER_PARAMETER_ARRAY(FIntVector4, Dilation_Stride_XBlockStartOffset_DilationXBlockStride, [FConvConstants::MAX_NUM_DIMENSIONS])
+			SHADER_PARAMETER_ARRAY(FIntVector4, GroupStride_GroupShape_GroupThreadStride_StrideXBlockStride, [FConvConstants::MAX_NUM_DIMENSIONS])
+			SHADER_PARAMETER_ARRAY(FIntVector4, YDimension_YMemoryStride_XDimension_XMemoryStride, [FConvConstants::MAX_NUM_DIMENSIONS])
+			SHADER_PARAMETER_ARRAY(FIntVector4, XBlockStartStride_XBlockStride_WDimension_WDimensionDilationXBlockStride, [FConvConstants::MAX_NUM_DIMENSIONS])
+			SHADER_PARAMETER_ARRAY(FVector4f, OneDiv_GroupStride_GroupThreadStride_XBlockStride, [FConvConstants::MAX_NUM_DIMENSIONS])
+			SHADER_PARAMETER(int32, NumWChannels)
+			SHADER_PARAMETER(int32, YBatchStride)
+			SHADER_PARAMETER(int32, YOutputKernelStride)
+			SHADER_PARAMETER(int32, XBatchStride)
+			SHADER_PARAMETER(int32, XChannelStride)
+			SHADER_PARAMETER(int32, XBlockSize)
+			SHADER_PARAMETER(int32, NumChannelBatches)
+			SHADER_PARAMETER(int32, NumChannelsPerBatch)
+			SHADER_PARAMETER(int32, WOutputKernelStride)
+			SHADER_PARAMETER(int32, WChannelBatchSize)
+			SHADER_PARAMETER(int32, WChannelSize)
+			SHADER_PARAMETER(float, GroupsDivM)
+		END_SHADER_PARAMETER_STRUCT()
+
+		static void ModifyCompilationEnvironment(const FGlobalShaderPermutationParameters& InParameters, FShaderCompilerEnvironment& OutEnvironment);
+
+	public:
+		static TArray<int32> GetOutputShape(TArrayView<const uint32> XShape, TArrayView<const uint32> WShape, EConvAutoPad AutoPad, TArrayView<const int32> Dilations, TArrayView<const int32> Strides, TArrayView<const int32> Pads);
+
+		static void FillInParameters(EConvGroupSize GroupSize, TArrayView<const uint32> XShape, TArrayView<const uint32> WShape, bool HasB,
+				EConvAutoPad AutoPad, int Group, TArrayView<const int32> Dilations, TArrayView<const int32> Strides, TArrayView<const int32> Pads, FConvCS::FParameters& Parameters);
+
+		static int32 GetNumReadsPerThread(EConvGroupSize GroupSize, TArrayView<const uint32> WShape, TArrayView<const int32> Dilations, TArrayView<const int32> Strides);
+
+		/**
+		* @brief Computes the group shape such that all dimension have roughly equal sizes.
+		*
+		* @param GroupSize The enum indicating the number of threads contained by a single group.
+		* @param NumDimensions The number of dimensions.
+		* @return TArray<int32> An array of size \p NumDimensions containing the number of threads in each dimension to form a volume of a total number of threads indicated by \p GroupSize
+		*/
+		static TArray<int32> GetGroupShape(EConvGroupSize GroupSize, int32 NumDimensions);
+
+		/**
+		* @brief Get the group count vector used to launch the gpu shader thread groups
+		*
+		* @param YShape The shape of the output as computed by GetOutputShape()
+		* @param YShape The shape of a single thread group as computed by GetGroupShape()
+		* @return FIntVector The number of thread groups to instantiate. z corresponds to the batch and y to the output kernel.
+		*/
+		static FIntVector GetGroupCount(TArrayView<const int32> YShape, TArrayView<const int32> GroupShape);
+
+		static EConvGroupSize GetMinimalGroupSize(TArrayView<const int32> WShape);
+	};
+} // UE::NNIHlslShaders::Internal
\ No newline at end of file
diff --git a/Engine/Plugins/Experimental/NNX/Source/NNIHlslShaders/Internal/NNIHlslShadersConvTransposeCS.h b/Engine/Plugins/Experimental/NNX/Source/NNIHlslShaders/Internal/NNIHlslShadersConvTransposeCS.h
new file mode 100644
--- /dev/null
+++ b/Engine/Plugins/Experimental/NNX/Source/NNIHlslShaders/Internal/NNIHlslShadersConvTransposeCS.h
@@ -0,0 +1,108 @@
+// Copyright Epic Games, Inc. All Rights Reserved.
+
+#pragma once
+
+#include "GlobalShader.h"
+#include "ShaderParameterUtils.h"
+#include "RenderGraphUtils.h"
+#include "NNXTypes.h"
+
+namespace UE::NNIHlslShaders::Internal
+{
+	enum class EConvTransposeAlgorithm : uint8
+	{
+		SharedMemory = 0,
+		MAX
+	};
+
+	enum class EConvTransposeGroupSize : uint8
+	{
+		Size128 = 0,
+		Size256,
+		Size512,
+		MAX
+	};
+
+	enum class EConvTransposeAutoPad : uint8
+	{
+		NOTSET = 0,// Use pad values passed in the array
+		SAME_UPPER,// Auto-pad to match input and output shape with potetnial extra padding at the end
+		SAME_LOWER,// Auto-pad to match input and output shape with potetnial extra padding at the beginning
+		VALID,// Set all paddings to zero
+		MAX
+	};
+
+	class FConvTransposeConstants
+	{
+	public:
+
+		static const int32 MAX_NUM_DIMENSIONS{4};
+		static const int32 MIN_NUM_READS_PER_THREAD_POW2{1};
+		static const int32 MAX_NUM_READS_PER_THREAD_POW2{3};
+	};
+
+	class NNIHLSLSHADERS_API FConvTransposeCS : public FGlobalShader
+	{
+		DECLARE_GLOBAL_SHADER(FConvTransposeCS);
+		SHADER_USE_PARAMETER_STRUCT(FConvTransposeCS, FGlobalShader);
+
+		class FConvTransposeAlgorithm : SHADER_PERMUTATION_ENUM_CLASS("ALGORITHM", EConvTransposeAlgorithm);
+		class FConvTransposeGroupSize : SHADER_PERMUTATION_ENUM_CLASS("GROUP_SIZE", EConvTransposeGroupSize);
+		class FConvTransposeNumStackDimensions : SHADER_PERMUTATION_RANGE_INT("NUM_STACK_DIMENSIONS", 1, FConvTransposeConstants::MAX_NUM_DIMENSIONS);
+		class FConvTransposeNumReadsPerThread : SHADER_PERMUTATION_RANGE_INT("NUM_READS_PER_THREAD_POW2", FConvTransposeConstants::MIN_NUM_READS_PER_THREAD_POW2, FConvTransposeConstants::MAX_NUM_READS_PER_THREAD_POW2);
+		class FConvTransposeHasB : SHADER_PERMUTATION_BOOL("HAS_B");
+		using FPermutationDomain = TShaderPermutationDomain<FConvTransposeAlgorithm, FConvTransposeGroupSize, FConvTransposeNumStackDimensions, FConvTransposeNumReadsPerThread, FConvTransposeHasB>;
+
+	public:
+
+		BEGIN_SHADER_PARAMETER_STRUCT(FParameters, )
+			SHADER_PARAMETER_RDG_BUFFER_SRV(Buffer<float>, X)
+			SHADER_PARAMETER_RDG_BUFFER_SRV(Buffer<float>, W)
+			SHADER_PARAMETER_RDG_BUFFER_UAV(RWBuffer<float>, Y)
+			SHADER_PARAMETER_RDG_BUFFER_SRV(Buffer<float>, B)
+			SHADER_PARAMETER_ARRAY(FIntVector4, Dilation_Stride_XBlockStartOffset_DilationXBlockStride, [FConvTransposeConstants::MAX_NUM_DIMENSIONS])
+			SHADER_PARAMETER_ARRAY(FIntVector4, GroupStride_GroupShape_GroupThreadStride_StrideXBlockStride, [FConvTransposeConstants::MAX_NUM_DIMENSIONS])
+			SHADER_PARAMETER_ARRAY(FIntVector4, YDimension_YMemoryStride_XDimension_XMemoryStride, [FConvTransposeConstants::MAX_NUM_DIMENSIONS])
+			SHADER_PARAMETER_ARRAY(FIntVector4, XBlockStartStride_XBlockStride_WDimension_WDimensionDilationXBlockStride, [FConvTransposeConstants::MAX_NUM_DIMENSIONS])
+			SHADER_PARAMETER_ARRAY(FVector4f, OneDiv_GroupStride_GroupThreadStride_OneDivStride, [FConvTransposeConstants::MAX_NUM_DIMENSIONS])
+			SHADER_PARAMETER(int32, NumWChannels)
+			SHADER_PARAMETER(int32, NumOutChannelsDivGroup)
+			SHADER_PARAMETER(int32, YBatchStride)
+			SHADER_PARAMETER(int32, YOutputKernelStride)
+			SHADER_PARAMETER(int32, XBatchStride)
+			SHADER_PARAMETER(int32, XChannelStride)
+			SHADER_PARAMETER(int32, XBlockSize)
+			SHADER_PARAMETER(int32, NumChannelBatches)
+			SHADER_PARAMETER(int32, NumChannelsPerBatch)
+			SHADER_PARAMETER(int32, WOutputKernelStride)
+			SHADER_PARAMETER(int32, WChannelBatchSize)
+			SHADER_PARAMETER(int32, WChannelSize)
+			SHADER_PARAMETER(float, GroupsDivM)
+			SHADER_PARAMETER(float, OneDivGroup)
+		END_SHADER_PARAMETER_STRUCT()
+
+		static void ModifyCompilationEnvironment(const FGlobalShaderPermutationParameters& InParameters, FShaderCompilerEnvironment& OutEnvironment);
+
+		static TArray<int32> GetOutputShape(TArrayView<const uint32> XShape, TArrayView<const uint32> WShape, EConvTransposeAutoPad AutoPad, TArrayView<const int32> Dilations, TArrayView<const int32> Strides, TArrayView<const int32> Pads, TArrayView<const int32> OutputPadding, int32 Group);
+
+		static void FillInParameters(EConvTransposeGroupSize GroupSize, TArrayView<const uint32> XShape, TArrayView<const uint32> WShape, bool HasB, EConvTransposeAutoPad AutoPad, int32 Group, TArrayView<const int32> Dilations, TArrayView<const int32> Strides, TArrayView<const int32> Pads, TArrayView<const int32> OutputPadding, FConvTransposeCS::FParameters& Parameters);
+
+		static int32 GetNumReadsPerThread(EConvTransposeGroupSize GroupSize, TArrayView<const uint32> WShape, TArrayView<const int32> Dilations, TArrayView<const int32> Strides);
+
+		static TArray<int32> GetGroupShape(EConvTransposeGroupSize GroupSize, int32 NumDimensions);
+
+		static FIntVector GetGroupCount(TArrayView<const int32> YShape, TArrayView<const int32> GroupShape);
+
+		static EConvTransposeGroupSize GetMinimalGroupSize(TArrayView<const uint32> WShape);
+
+	private:
+
+		static TArray<int32> GetXBlockShape(TArrayView<const int32> GroupShape, TArrayView<const uint32> WShape, TArrayView<const int32> Dilations, TArrayView<const int32> Strides);
+
+		static TArray<int32> GetPadding(TArrayView<const uint32> WShape, EConvTransposeAutoPad AutoPad, TArrayView<const int32> Dilations, TArrayView<const int32> Strides, TArrayView<const int32> Pads, TArrayView<const int32> OutputPadding);
+
+		static int32 GetNumThreadsPerGroup(EConvTransposeGroupSize GroupSize);
+
+		static TArray<int32> GetGridShape(TArrayView<const int32> YShape, TArrayView<const int32> GroupShape);
+	};
+} // UE::NNIHlslShaders::Internal
\ No newline at end of file
diff --git a/Engine/Plugins/Experimental/NNX/Source/NNIHlslShaders/Internal/NNIHlslShadersElementWiseBinaryCS.h b/Engine/Plugins/Experimental/NNX/Source/NNIHlslShaders/Internal/NNIHlslShadersElementWiseBinaryCS.h
new file mode 100644
--- /dev/null
+++ b/Engine/Plugins/Experimental/NNX/Source/NNIHlslShaders/Internal/NNIHlslShadersElementWiseBinaryCS.h
@@ -0,0 +1,46 @@
+// Copyright Epic Games, Inc. All Rights Reserved.
+
+#pragma once
+
+#include "GlobalShader.h"
+#include "ShaderParameterUtils.h"
+#include "RenderGraphUtils.h"
+#include "NNXOperator.h"
+
+namespace UE::NNIHlslShaders::Internal
+{
+	class FElementWiseBinaryConstants
+	{
+	public:
+
+		static const int32 MAX_NUM_DIMENSIONS{ 8 };
+		static const int32 NUM_GROUP_THREADS{ 256 };
+	};
+
+	class NNIHLSLSHADERS_API TElementWiseBinaryCS : public FGlobalShader
+	{
+		DECLARE_GLOBAL_SHADER(TElementWiseBinaryCS);
+		SHADER_USE_PARAMETER_STRUCT(TElementWiseBinaryCS, FGlobalShader)
+
+		class FOperatorType : SHADER_PERMUTATION_ENUM_CLASS("OP_TYPENAME", EMLElementWiseBinaryOperatorType);
+		class FBinaryNumDimensions : SHADER_PERMUTATION_RANGE_INT("NUM_DIMENSIONS", 1, FElementWiseBinaryConstants::MAX_NUM_DIMENSIONS);
+		using FPermutationDomain = TShaderPermutationDomain<FOperatorType, FBinaryNumDimensions>;
+
+	public:
+
+		BEGIN_SHADER_PARAMETER_STRUCT(FParameters, )
+			SHADER_PARAMETER_RDG_BUFFER_SRV(Buffer<float>, LHSInput)
+			SHADER_PARAMETER_RDG_BUFFER_SRV(Buffer<float>, RHSInput)
+			SHADER_PARAMETER_RDG_BUFFER_UAV(RWBuffer<float>, Output)
+			SHADER_PARAMETER_ARRAY(FUintVector4, TensorInfo, [FElementWiseBinaryConstants::MAX_NUM_DIMENSIONS])
+			SHADER_PARAMETER(uint32, Num)
+			SHADER_PARAMETER(uint32, ThreadCountX)
+		END_SHADER_PARAMETER_STRUCT()
+
+		static void ModifyCompilationEnvironment(const FGlobalShaderPermutationParameters& InParameters, FShaderCompilerEnvironment& OutEnvironment);
+
+	private:
+
+		static const FString GetOpFunc(EMLElementWiseBinaryOperatorType OpType);
+	};
+} // UE::NNIHlslShaders::Internal
diff --git a/Engine/Plugins/Experimental/NNX/Source/NNIHlslShaders/Internal/NNIHlslShadersElementWiseUnaryCS.h b/Engine/Plugins/Experimental/NNX/Source/NNIHlslShaders/Internal/NNIHlslShadersElementWiseUnaryCS.h
new file mode 100644
--- /dev/null
+++ b/Engine/Plugins/Experimental/NNX/Source/NNIHlslShaders/Internal/NNIHlslShadersElementWiseUnaryCS.h
@@ -0,0 +1,45 @@
+// Copyright Epic Games, Inc. All Rights Reserved.
+
+#pragma once
+
+#include "GlobalShader.h"
+#include "ShaderParameterUtils.h"
+#include "RenderGraphUtils.h"
+#include "NNXOperator.h"
+
+namespace UE::NNIHlslShaders::Internal
+{
+	class FElementWiseUnaryConstants
+	{
+	public:
+
+		static const int32 NUM_GROUP_THREADS{ 256 };
+	};
+
+	class NNIHLSLSHADERS_API TElementWiseUnaryCS : public FGlobalShader
+	{
+		DECLARE_GLOBAL_SHADER(TElementWiseUnaryCS);
+		SHADER_USE_PARAMETER_STRUCT(TElementWiseUnaryCS, FGlobalShader)
+
+		class FOperatorType : SHADER_PERMUTATION_ENUM_CLASS("OP_TYPENAME", EMLElementWiseUnaryOperatorType);
+		using FPermutationDomain = TShaderPermutationDomain<FOperatorType>;
+
+	public:
+
+		BEGIN_SHADER_PARAMETER_STRUCT(FParameters, )
+			SHADER_PARAMETER_RDG_BUFFER_SRV(Buffer<float>, Input)
+			SHADER_PARAMETER_RDG_BUFFER_UAV(RWBuffer<float>, Output)
+			SHADER_PARAMETER(uint32, Num)
+			SHADER_PARAMETER(uint32, ThreadCountX)
+			SHADER_PARAMETER(float, Alpha)
+			SHADER_PARAMETER(float, Beta)
+			SHADER_PARAMETER(float, Gamma)
+		END_SHADER_PARAMETER_STRUCT()
+
+		static void ModifyCompilationEnvironment(const FGlobalShaderPermutationParameters& InParameters, FShaderCompilerEnvironment& OutEnvironment);
+
+	private:
+
+		static const FString GetOpFunc(EMLElementWiseUnaryOperatorType OpType);
+	};
+} // UE::NNIHlslShaders::Internal
diff --git a/Engine/Plugins/Experimental/NNX/Source/NNIHlslShaders/Internal/NNIHlslShadersElementWiseVariadicCS.h b/Engine/Plugins/Experimental/NNX/Source/NNIHlslShaders/Internal/NNIHlslShadersElementWiseVariadicCS.h
new file mode 100644
--- /dev/null
+++ b/Engine/Plugins/Experimental/NNX/Source/NNIHlslShaders/Internal/NNIHlslShadersElementWiseVariadicCS.h
@@ -0,0 +1,54 @@
+// Copyright Epic Games, Inc. All Rights Reserved.
+
+#pragma once
+
+#include "GlobalShader.h"
+#include "ShaderParameterUtils.h"
+#include "RenderGraphUtils.h"
+#include "NNXOperator.h"
+
+namespace UE::NNIHlslShaders::Internal
+{
+	class FElementWiseVariadicConstants
+	{
+	public:
+
+		static const int32 MAX_NUM_DIMENSIONS{ 8 };
+		static const uint32 NUM_GROUP_THREADS{ 256 };
+		static const uint32 MAX_NUM_INPUT{ 4 };
+	};
+
+	class NNIHLSLSHADERS_API TElementWiseVariadicCS : public FGlobalShader
+	{
+		DECLARE_GLOBAL_SHADER(TElementWiseVariadicCS);
+		SHADER_USE_PARAMETER_STRUCT(TElementWiseVariadicCS, FGlobalShader)
+
+		class FOperatorType : SHADER_PERMUTATION_ENUM_CLASS("OP_TYPENAME", EMLElementWiseVariadicOperatorType);
+		class FApplyScale : SHADER_PERMUTATION_BOOL("APPLYSCALE");
+		class FOutputAsInput : SHADER_PERMUTATION_BOOL("OUTPUTASINPUT");
+		class FNumInput : SHADER_PERMUTATION_RANGE_INT("NUMINPUT", 1, FElementWiseVariadicConstants::MAX_NUM_INPUT);
+		class FVariadicNumDimensions : SHADER_PERMUTATION_RANGE_INT("NUM_DIMENSIONS", 1, FElementWiseVariadicConstants::MAX_NUM_DIMENSIONS);
+		using FPermutationDomain = TShaderPermutationDomain<FOperatorType, FApplyScale, FOutputAsInput, FNumInput, FVariadicNumDimensions>;
+
+	public:
+
+		BEGIN_SHADER_PARAMETER_STRUCT(FParameters, )
+			SHADER_PARAMETER_RDG_BUFFER_SRV(Buffer<float>, Input0)
+			SHADER_PARAMETER_RDG_BUFFER_SRV(Buffer<float>, Input1)
+			SHADER_PARAMETER_RDG_BUFFER_SRV(Buffer<float>, Input2)
+			SHADER_PARAMETER_RDG_BUFFER_SRV(Buffer<float>, Input3)
+			SHADER_PARAMETER_RDG_BUFFER_UAV(RWBuffer<float>, Output)
+			SHADER_PARAMETER_ARRAY(FUintVector4, InputTensorInfo, [FElementWiseVariadicConstants::MAX_NUM_DIMENSIONS])
+			SHADER_PARAMETER_ARRAY(FUintVector4, OutputTensorInfo, [FElementWiseVariadicConstants::MAX_NUM_DIMENSIONS])
+			SHADER_PARAMETER(uint32, Num)
+			SHADER_PARAMETER(uint32, ThreadCountX)
+			SHADER_PARAMETER(float, Scale)
+		END_SHADER_PARAMETER_STRUCT()
+
+		static void ModifyCompilationEnvironment(const FGlobalShaderPermutationParameters& InParameters, FShaderCompilerEnvironment& OutEnvironment);
+
+	private:
+
+		static const FString GetOpFunc(EMLElementWiseVariadicOperatorType OpType);
+	};
+} // UE::NNIHlslShaders::Internal
\ No newline at end of file
diff --git a/Engine/Plugins/Experimental/NNX/Source/NNIHlslShaders/Internal/NNIHlslShadersGatherCS.h b/Engine/Plugins/Experimental/NNX/Source/NNIHlslShaders/Internal/NNIHlslShadersGatherCS.h
new file mode 100644
--- /dev/null
+++ b/Engine/Plugins/Experimental/NNX/Source/NNIHlslShaders/Internal/NNIHlslShadersGatherCS.h
@@ -0,0 +1,49 @@
+// Copyright Epic Games, Inc. All Rights Reserved.
+
+#pragma once
+
+#include "GlobalShader.h"
+#include "ShaderParameterUtils.h"
+#include "RenderGraphUtils.h"
+#include "NNXTypes.h"
+
+namespace UE::NNIHlslShaders::Internal
+{
+	class FGatherConstants
+	{
+	public:
+
+		static const int32 MAX_NUM_DIMENSIONS{ 8 };
+		static const int32 NUM_GROUP_THREADS{ 256 };
+	};
+
+	// template <typename DataElementType, typename IndicesElementType>
+	class NNIHLSLSHADERS_API TGatherCS : public FGlobalShader
+	{
+		DECLARE_GLOBAL_SHADER(TGatherCS);
+		SHADER_USE_PARAMETER_STRUCT(TGatherCS, FGlobalShader)
+
+		class FGatherNumOutputDimensions : SHADER_PERMUTATION_RANGE_INT("NUM_OUTPUT_DIMENSIONS", 1, FGatherConstants::MAX_NUM_DIMENSIONS);
+		using FPermutationDomain = TShaderPermutationDomain<FGatherNumOutputDimensions>;
+
+	public:
+
+		BEGIN_SHADER_PARAMETER_STRUCT(FParameters, )
+			SHADER_PARAMETER(int32, Axis)
+			SHADER_PARAMETER(int32, OutputSize)
+			SHADER_PARAMETER(int32, NumDataDimensions)
+			SHADER_PARAMETER(int32, NumIndicesDimensions)
+			SHADER_PARAMETER_ARRAY(FIntVector4, DataStride_IndicesStride_OutputStride, [FGatherConstants::MAX_NUM_DIMENSIONS])
+			SHADER_PARAMETER_ARRAY(FVector4f, OneDivDataStride_OneDivIndicesStride_OneDivOutputStride, [FGatherConstants::MAX_NUM_DIMENSIONS])
+			SHADER_PARAMETER_RDG_BUFFER_SRV(Buffer<float>, Data)
+			SHADER_PARAMETER_RDG_BUFFER_SRV(Buffer<int32>, Indices)
+			SHADER_PARAMETER_RDG_BUFFER_UAV(RWBuffer<float>, Output)
+		END_SHADER_PARAMETER_STRUCT()
+
+		static void ModifyCompilationEnvironment(const FGlobalShaderPermutationParameters& InParameters, FShaderCompilerEnvironment& OutEnvironment);
+
+		static void FillInParameters(int32 Axis, const NNX::FMLTensorDesc& Data, const NNX::FMLTensorDesc& Indices, FParameters& Parameters);
+
+		static FIntVector GetGroupCount(const FParameters& Parameters);
+	};
+} // UE::NNIHlslShaders::Internal
\ No newline at end of file
diff --git a/Engine/Plugins/Experimental/NNX/Source/NNIHlslShaders/Internal/NNIHlslShadersGemmCS.h b/Engine/Plugins/Experimental/NNX/Source/NNIHlslShaders/Internal/NNIHlslShadersGemmCS.h
new file mode 100644
--- /dev/null
+++ b/Engine/Plugins/Experimental/NNX/Source/NNIHlslShaders/Internal/NNIHlslShadersGemmCS.h
@@ -0,0 +1,89 @@
+// Copyright Epic Games, Inc. All Rights Reserved.
+
+#pragma once
+
+#include "GlobalShader.h"
+#include "ShaderParameterUtils.h"
+#include "RenderGraphUtils.h"
+#include "NNXTypes.h"
+
+namespace UE::NNIHlslShaders::Internal
+{
+	enum class EGemmCScalar : uint8
+	{
+		No = 0,
+		Yes,
+		NoBias,
+		MAX
+	};
+
+	enum class EGemmAlgorithm : uint8
+	{
+		Simple8x8 = 0,
+		Simple16x16,
+		Simple32x32,
+		Simple256x1,
+		SharedMemory8x8,
+		SharedMemory16x16,
+		SharedMemory32x32,
+		MultiWrite1x16,
+		MultiWrite2x16,
+		MultiWrite1x32,
+		MultiWrite2x32,
+		MultiWrite4x32,
+		MultiWrite2x64,
+		MultiWrite4x64,
+		MultiWrite8x64,
+		MAX
+	};
+
+	class FGemmConstants
+	{
+	public:
+
+		static const int32 MAX_NUM_STACK_DIMENSIONS{8};
+	};
+
+	class NNIHLSLSHADERS_API TGemmCS : public FGlobalShader
+	{
+		DECLARE_GLOBAL_SHADER(TGemmCS);
+		SHADER_USE_PARAMETER_STRUCT(TGemmCS, FGlobalShader)
+
+		class FGemmCScalar : SHADER_PERMUTATION_ENUM_CLASS("C_SCALAR", EGemmCScalar);
+		class FGemmAlgorithm : SHADER_PERMUTATION_ENUM_CLASS("ALGORITHM", EGemmAlgorithm);
+		class FGemmNumStackDimensions : SHADER_PERMUTATION_RANGE_INT("NUM_STACK_DIMENSIONS", 0, FGemmConstants::MAX_NUM_STACK_DIMENSIONS);
+		using FPermutationDomain = TShaderPermutationDomain<FGemmCScalar, FGemmAlgorithm, FGemmNumStackDimensions>;
+
+	public:
+
+		BEGIN_SHADER_PARAMETER_STRUCT(FParameters, )
+			SHADER_PARAMETER(float, Alpha)
+			SHADER_PARAMETER(float, Beta)
+			SHADER_PARAMETER(int32, TransA)
+			SHADER_PARAMETER(int32, TransB)
+			SHADER_PARAMETER(uint32, M)
+			SHADER_PARAMETER(uint32, N)
+			SHADER_PARAMETER(uint32, K)
+			SHADER_PARAMETER(uint32, MxK)
+			SHADER_PARAMETER(uint32, KxN)
+			SHADER_PARAMETER(uint32, MxN)
+			SHADER_PARAMETER(uint32, CWidth)
+			SHADER_PARAMETER(uint32, CHeight)
+			SHADER_PARAMETER(float, CScalar)
+			SHADER_PARAMETER_ARRAY(FUint32Vector4, StackShapeA_StackShapeB_StackStrideA_StackStrideB, [FGemmConstants::MAX_NUM_STACK_DIMENSIONS])
+			SHADER_PARAMETER_RDG_BUFFER_SRV(Buffer<float>, A)
+			SHADER_PARAMETER_RDG_BUFFER_SRV(Buffer<float>, B)
+			SHADER_PARAMETER_RDG_BUFFER_SRV(Buffer<float>, C)
+			SHADER_PARAMETER_RDG_BUFFER_UAV(RWBuffer<float>, Y)
+		END_SHADER_PARAMETER_STRUCT()
+
+		static void ModifyCompilationEnvironment(const FGlobalShaderPermutationParameters& InParameters, FShaderCompilerEnvironment& OutEnvironment);
+
+		static void FillInParameters(float Alpha, float Beta, int32 TransA, int32 TransB, const NNX::FMLTensorDesc &InputA, const NNX::FMLTensorDesc &InputB, const NNX::FMLTensorDesc &InputC, float CScalar, FParameters& Parameters);
+
+		static void FillInParametersMatMul(const NNX::FMLTensorDesc &InputA, const NNX::FMLTensorDesc &InputB, FParameters& Parameters);
+
+		static FIntVector GetGroupCount(const FParameters& Parameters, EGemmAlgorithm Algorithm, int32 NumStackDimensions);
+		static EGemmAlgorithm GetAlgorithm(const FParameters& Parameters);
+	};
+} // UE::NNIHlslShaders::Internal
\ No newline at end of file
diff --git a/Engine/Plugins/Experimental/NNX/Source/NNIHlslShaders/NNIHlslShaders.Build.cs b/Engine/Plugins/Experimental/NNX/Source/NNIHlslShaders/NNIHlslShaders.Build.cs
new file mode 100644
--- /dev/null
+++ b/Engine/Plugins/Experimental/NNX/Source/NNIHlslShaders/NNIHlslShaders.Build.cs
@@ -0,0 +1,28 @@
+// Copyright Epic Games, Inc. All Rights Reserved.
+
+using UnrealBuildTool;
+
+public class NNIHlslShaders : ModuleRules
+{
+	public NNIHlslShaders( ReadOnlyTargetRules Target ) : base( Target )
+	{
+        PCHUsage = ModuleRules.PCHUsageMode.UseExplicitOrSharedPCHs;
+
+		PublicDependencyModuleNames.AddRange(
+			new string[] {
+				"Core",
+				"CoreUObject",
+				"Engine"
+			}
+		);
+
+		PrivateDependencyModuleNames.AddRange(
+			new string[] {
+				"Projects",
+				"RenderCore",
+				"RHI",
+				"NNXCore"
+			}
+		);
+	}
+}
diff --git a/Engine/Plugins/Experimental/NNX/Source/NNIHlslShaders/Private/NNIHlslShadersConvCS.cpp b/Engine/Plugins/Experimental/NNX/Source/NNIHlslShaders/Private/NNIHlslShadersConvCS.cpp
new file mode 100644
--- /dev/null
+++ b/Engine/Plugins/Experimental/NNX/Source/NNIHlslShaders/Private/NNIHlslShadersConvCS.cpp
@@ -0,0 +1,292 @@
+// Copyright Epic Games, Inc. All Rights Reserved.
+
+#include "NNIHlslShadersConvCS.h"
+#include "NNXCore.h"
+
+namespace UE::NNIHlslShaders::Internal
+{
+	namespace
+	{
+
+	TArray<int32> GetXBlockShape(TArrayView<const int32> GroupShape, TArrayView<const uint32> WShape, TArrayView<const int32> Dilations, TArrayView<const int32> Strides)
+	{
+		check(WShape.Num() > 2);
+		check(GroupShape.Num() == WShape.Num() - 2);
+		check(Dilations.Num() == 0 || Dilations.Num() == GroupShape.Num());
+		check(Strides.Num() == 0 || Strides.Num() == GroupShape.Num());
+		
+		TArray<int32> Result;
+		Result.SetNumUninitialized(GroupShape.Num());
+		for (int32 i = 0; i < GroupShape.Num(); i++)
+		{
+			int32 DilatedKernelSize = (i < Dilations.Num() ? Dilations[i] : 1) * (WShape[2 + i] - 1) + 1;
+			Result[i] = DilatedKernelSize + (GroupShape[i] - 1) * (i < Strides.Num() ? Strides[i] : 1); 
+		}
+		return Result;
+	}
+
+	TArray<int32> GetPadding(TArrayView<const uint32> XShape, TArrayView<const uint32> WShape, EConvAutoPad AutoPad, TArrayView<const int32> Dilations, TArrayView<const int32> Strides, TArrayView<const int32> Pads)
+	{
+		check(XShape.Num() > 2);
+		check(WShape.Num() == XShape.Num());
+		check(Dilations.Num() == 0 || Dilations.Num() == WShape.Num() - 2);
+		check(Strides.Num() == 0 || Strides.Num() == WShape.Num() - 2);
+		check(AutoPad != EConvAutoPad::NOTSET || Pads.Num() == 2 * (WShape.Num() - 2));
+
+		int32 NumDimensions = XShape.Num() - 2;
+		TArray<int32> Result;
+		Result.Init(0, 2 * NumDimensions);
+
+		if (AutoPad == EConvAutoPad::NOTSET)
+		{
+			return TArray<int32>{Pads};
+		}
+		else if (AutoPad == EConvAutoPad::VALID)
+		{
+			return Result;
+		}
+
+		for (int32 DimensionIndex = 0; DimensionIndex < NumDimensions; DimensionIndex++)
+		{
+			int32 DilatedKernelSize = (DimensionIndex < Dilations.Num() ? Dilations[DimensionIndex] : 1) * (WShape[DimensionIndex + 2] - 1) + 1;
+			int32 TotalPad = ((int32)(((XShape[DimensionIndex + 2] + (DimensionIndex < Strides.Num() ? Strides[DimensionIndex] : 1) - 1) / (DimensionIndex < Strides.Num() ? Strides[DimensionIndex] : 1)) - 1)) * (DimensionIndex < Strides.Num() ? Strides[DimensionIndex] : 1) + DilatedKernelSize - XShape[DimensionIndex + 2];
+			if (AutoPad == EConvAutoPad::SAME_LOWER)
+			{
+				Result[DimensionIndex] = (TotalPad + 1) / 2;
+			}	
+			else
+			{
+				Result[DimensionIndex] = TotalPad / 2;
+			}
+			Result[NumDimensions + DimensionIndex] = TotalPad - Result[DimensionIndex];
+		}
+
+		return Result;
+	}
+
+	int32 GetNumThreadsPerGroup(EConvGroupSize GroupSize)
+	{
+		int32 NumThreadsPerGroup = 128;
+		switch (GroupSize)
+		{
+			case EConvGroupSize::Size128:
+				NumThreadsPerGroup = 128;
+				break;
+			case EConvGroupSize::Size256:
+				NumThreadsPerGroup = 256;
+				break;
+			case EConvGroupSize::Size512:
+				NumThreadsPerGroup = 512;
+				break;
+			default:
+				NumThreadsPerGroup = 128;
+				break;
+		}
+		check(FMath::Log2((float)NumThreadsPerGroup) == FMath::Floor(FMath::Log2((float)NumThreadsPerGroup)));
+		return NumThreadsPerGroup;
+	}
+
+	TArray<int32> GetGridShape(TArrayView<const int32> YShape, TArrayView<const int32> GroupShape)
+	{
+		check(YShape.Num() > 2);
+		check(YShape.Num() == (GroupShape.Num() + 2));
+
+		TArray<int32> Result;
+		Result.SetNumUninitialized(YShape.Num() - 2);
+		for (int32 i = 2; i < YShape.Num(); i++)
+		{
+			Result[i-2] = FMath::DivideAndRoundUp(YShape[i], GroupShape[i-2]);
+		}
+
+		return Result;
+	}
+
+	} // namespace
+
+	void FConvCS::ModifyCompilationEnvironment(const FGlobalShaderPermutationParameters& InParameters, FShaderCompilerEnvironment& OutEnvironment)
+	{
+		FGlobalShader::ModifyCompilationEnvironment(InParameters, OutEnvironment);
+		OutEnvironment.SetDefine(TEXT("MAX_NUM_DIMENSIONS"), FConvConstants::MAX_NUM_DIMENSIONS);
+	}
+
+	TArray<int32> FConvCS::GetOutputShape(TArrayView<const uint32> XShape, TArrayView<const uint32> WShape, EConvAutoPad AutoPad, TArrayView<const int32> Dilations, TArrayView<const int32> Strides, TArrayView<const int32> Pads)
+	{
+		check(XShape.Num() > 2);
+		check(WShape.Num() == XShape.Num());
+		check(Dilations.Num() == 0 || Dilations.Num() == WShape.Num() - 2);
+		check(Strides.Num() == 0 || Strides.Num() == WShape.Num() - 2);
+		check(Pads.Num() == 0 || Pads.Num() == 2 * (WShape.Num() - 2));
+
+		TArray<int32> Padding = GetPadding(XShape, WShape, AutoPad, Dilations, Strides, Pads);
+
+		TArray<int32> Result;
+		Result.SetNumUninitialized(XShape.Num());
+		Result[0] = XShape[0];
+		Result[1] = WShape[0];
+
+		int32 NumDimensions = XShape.Num() - 2;
+		for (int32 DimensionIndex = 0; DimensionIndex < NumDimensions; DimensionIndex++)
+		{
+			int32 PaddedSize = XShape[2 + DimensionIndex] + Padding[DimensionIndex] + Padding[NumDimensions + DimensionIndex];
+			int32 DilatedKernelSize = (DimensionIndex < Dilations.Num() ? Dilations[DimensionIndex] : 1) * (WShape[2 + DimensionIndex] - 1) + 1;
+			float PaddeSizeMinusDilatedKernelSize = (float)(PaddedSize - DilatedKernelSize);
+			int32 OutputSize = (int32)(PaddeSizeMinusDilatedKernelSize / (DimensionIndex < Strides.Num() ? (float)Strides[DimensionIndex] : 1.0) + 1);
+			Result[2 + DimensionIndex] = OutputSize;
+		}
+
+		return Result;
+	}
+
+	void FConvCS::FillInParameters(EConvGroupSize GroupSize, TArrayView<const uint32> XShape, TArrayView<const uint32> WShape, bool HasB,
+			EConvAutoPad AutoPad, int Group, TArrayView<const int32> Dilations, TArrayView<const int32> Strides, TArrayView<const int32> Pads, FConvCS::FParameters& Parameters)
+	{
+		check(XShape.Num() > 2);
+		check(WShape.Num() == XShape.Num());
+		check(Dilations.Num() == 0 || Dilations.Num() == WShape.Num() - 2);
+		check(Strides.Num() == 0 || Strides.Num() == WShape.Num() - 2);
+		check(Pads.Num() == 0 || Pads.Num() == 2 * (WShape.Num() - 2));
+		check(GetNumReadsPerThread(GroupSize, WShape, Dilations, Strides) >= 0)
+		check(WShape[0] > 0)
+		check(WShape[1] > 0)
+
+		int32 NumDimensions = XShape.Num() - 2;
+
+		TArray<int32> Padding = GetPadding(XShape, WShape, AutoPad, Dilations, Strides, Pads);
+		TArray<int32> GroupShape = GetGroupShape(GroupSize, NumDimensions);
+		TArray<int32> YShape = GetOutputShape(XShape, WShape, AutoPad, Dilations, Strides, Pads);
+		TArray<int32> GridShape = GetGridShape(YShape, GroupShape);
+		TArray<int32> XBlockShape = GetXBlockShape(GroupShape, WShape, Dilations, Strides);
+		
+		int32 GroupStride = 1;
+		int32 GroupThreadStride = 1;
+		int32 XBlockSize = 1;
+		int32 YMemoryStride = 1;
+		int32 XMemoryStride = 1;
+		int32 WChannelSize = 1;
+		for (int32 i = NumDimensions-1; i >= 0; i--)
+		{
+			int32 Stride = i < Strides.Num() ? Strides[i] : 1;
+			int32 Dilation = i < Dilations.Num() ? Dilations[i] : 1;
+			Parameters.Dilation_Stride_XBlockStartOffset_DilationXBlockStride[i] = FIntVector4(Dilation, Stride, -Padding[i], Dilation * XBlockSize);
+			Parameters.GroupStride_GroupShape_GroupThreadStride_StrideXBlockStride[i] = FIntVector4(GroupStride, GroupShape[i], GroupThreadStride, Stride * XBlockSize);
+			Parameters.YDimension_YMemoryStride_XDimension_XMemoryStride[i] = FIntVector4(YShape[2 + i], YMemoryStride, XShape[2 + i], XMemoryStride);
+			Parameters.XBlockStartStride_XBlockStride_WDimension_WDimensionDilationXBlockStride[i] = FIntVector4(GroupShape[i] * Stride, XBlockSize, WShape[2 + i], WShape[2 + i] * Dilation * XBlockSize);
+			Parameters.OneDiv_GroupStride_GroupThreadStride_XBlockStride[i] = FVector4f(1.0/((float)GroupStride), 1.0/((float)GroupThreadStride), 1.0/((float)XBlockSize), 0.0);
+
+			GroupStride *= GridShape[i];
+			GroupThreadStride *= GroupShape[i];
+			XBlockSize *= XBlockShape[i];
+			YMemoryStride *= YShape[2 + i];
+			XMemoryStride *= XShape[2 + i];
+			WChannelSize *= WShape[2 + i];
+		}
+		
+		Parameters.NumWChannels = WShape[1];
+
+		Parameters.YBatchStride = YShape[1] * YMemoryStride;
+		Parameters.YOutputKernelStride = YMemoryStride;
+		
+		Parameters.XBatchStride = XShape[1] * XMemoryStride;
+		Parameters.XChannelStride = XMemoryStride;
+
+		Parameters.XBlockSize = XBlockSize;
+
+		Parameters.NumChannelsPerBatch = FMath::Min((int32)((float)GroupThreadStride / (float)WChannelSize), (int32)WShape[1]);
+		check(Parameters.NumChannelsPerBatch > 0)
+		Parameters.NumChannelBatches = FMath::DivideAndRoundUp((int32)WShape[1], Parameters.NumChannelsPerBatch);
+		
+		Parameters.WOutputKernelStride = WShape[1] * WChannelSize;
+		Parameters.WChannelBatchSize = Parameters.NumChannelsPerBatch * WChannelSize;
+		Parameters.WChannelSize = WChannelSize;
+
+		Parameters.GroupsDivM = (float)Group / (float)WShape[0];
+	}
+
+	int32 FConvCS::GetNumReadsPerThread(EConvGroupSize GroupSize, TArrayView<const uint32> WShape, TArrayView<const int32> Dilations, TArrayView<const int32> Strides)
+	{
+		check(WShape.Num() > 2);
+		check(Dilations.Num() == 0 || Dilations.Num() == WShape.Num() - 2);
+		check(Strides.Num() == 0 || Strides.Num() == WShape.Num() - 2);
+
+		int32 NumDimensions = WShape.Num() - 2;
+
+		TArray<int32> GroupShape = GetGroupShape(GroupSize, NumDimensions);
+		int32 NumThreadsPerGroup = 1;
+		for (int32 i = 0; i < NumDimensions; i++)
+		{
+			NumThreadsPerGroup *= GroupShape[i];
+		}
+
+		TArray<int32> XBlockShape = GetXBlockShape(GroupShape, WShape, Dilations, Strides);
+		int32 NumXBlockElements = 1;
+		for (int32 i = 0; i < NumDimensions; i++)
+		{
+			NumXBlockElements *= XBlockShape[i];
+		}
+
+		int32 NumReads = FMath::DivideAndRoundUp(NumXBlockElements, NumThreadsPerGroup);
+		int32 NumReadsPow2 = FMath::Max(FMath::RoundToPositiveInfinity(FMath::Log2((float)NumReads)), FConvConstants::MIN_NUM_READS_PER_THREAD_POW2);
+
+		if(NumReadsPow2 <= FConvConstants::MAX_NUM_READS_PER_THREAD_POW2)
+		{
+			return NumReadsPow2;
+		}
+
+		return -1;
+	}
+
+	TArray<int32> FConvCS::GetGroupShape(EConvGroupSize GroupSize, int32 NumDimensions)
+	{
+		check(NumDimensions > 0);
+
+		int32 NumThreadsPerGroup = GetNumThreadsPerGroup(GroupSize);
+
+		int32 Power = (int32)FMath::Log2((float)NumThreadsPerGroup);
+		int32 MinPowerPerDim = (int32)((float)Power / (float)NumDimensions);
+		int32 PowerReminder = Power - NumDimensions * MinPowerPerDim;
+
+		TArray<int32> Result;
+		Result.Init((int32)FMath::Pow((float)2.0, (float)MinPowerPerDim), NumDimensions);
+		for (int32 i = 0; i < PowerReminder; i++)
+		{
+			Result[NumDimensions-1-i] *= 2;
+		}
+		return Result;
+	}
+
+	FIntVector FConvCS::GetGroupCount(TArrayView<const int32> YShape, TArrayView<const int32> GroupShape)
+	{
+		check(YShape.Num() > 2);
+		check(YShape.Num() == (GroupShape.Num() + 2));
+
+		int32 ThreadGroupCountValueX = 1;
+		for (int32 i = 2; i < YShape.Num(); i++)
+		{
+			ThreadGroupCountValueX *= FMath::DivideAndRoundUp(YShape[i], GroupShape[i-2]);
+		}
+
+		return FIntVector(ThreadGroupCountValueX, YShape[1], YShape[0]);
+	}
+
+	EConvGroupSize FConvCS::GetMinimalGroupSize(TArrayView<const int32> WShape)
+	{
+		int32 NumDimensions = WShape.Num() - 2;
+		int32 WChannelSize = 1;
+		for (int32 i = 0; i < NumDimensions; i++)
+		{
+			WChannelSize *= WShape[2 + i];
+		}
+		
+		for (int32 i = 0; i < (int32)EConvGroupSize::MAX; i++)
+		{
+			if (GetNumThreadsPerGroup((EConvGroupSize)i) >= WChannelSize) 
+			{
+				return (EConvGroupSize)i;
+			}
+		}
+		
+		return EConvGroupSize::MAX;
+	}
+
+	IMPLEMENT_GLOBAL_SHADER(FConvCS, "/NNI/NNIHlslShadersConv.usf", "Conv", SF_Compute);
+} // UE::NNIHlslShaders::Internal
\ No newline at end of file
diff --git a/Engine/Plugins/Experimental/NNX/Source/NNIHlslShaders/Private/NNIHlslShadersConvTransposeCS.cpp b/Engine/Plugins/Experimental/NNX/Source/NNIHlslShaders/Private/NNIHlslShadersConvTransposeCS.cpp
new file mode 100644
--- /dev/null
+++ b/Engine/Plugins/Experimental/NNX/Source/NNIHlslShaders/Private/NNIHlslShadersConvTransposeCS.cpp
@@ -0,0 +1,285 @@
+// Copyright Epic Games, Inc. All Rights Reserved.
+
+#include "NNIHlslShadersConvTransposeCS.h"
+
+namespace UE::NNIHlslShaders::Internal
+{
+	void FConvTransposeCS::ModifyCompilationEnvironment(const FGlobalShaderPermutationParameters& InParameters, FShaderCompilerEnvironment& OutEnvironment)
+	{
+		FGlobalShader::ModifyCompilationEnvironment(InParameters, OutEnvironment);
+		OutEnvironment.SetDefine(TEXT("MAX_NUM_STACK_DIMENSIONS"), FConvTransposeConstants::MAX_NUM_DIMENSIONS);
+	}
+
+	TArray<int32> FConvTransposeCS::GetOutputShape(TArrayView<const uint32> XShape, TArrayView<const uint32> WShape, EConvTransposeAutoPad AutoPad, TArrayView<const int32> Dilations, TArrayView<const int32> Strides, TArrayView<const int32> Pads, TArrayView<const int32> OutputPadding, int32 Group)
+	{
+		check(XShape.Num() > 2);
+		check(WShape.Num() == XShape.Num());
+		check(Dilations.Num() == 0 || Dilations.Num() == WShape.Num() - 2);
+		check(Strides.Num() == 0 || Strides.Num() == WShape.Num() - 2);
+		check(AutoPad != EConvTransposeAutoPad::NOTSET || Pads.Num() == 2 * (WShape.Num() - 2));
+
+		TArray<int32> Padding = GetPadding(WShape, AutoPad, Dilations, Strides, Pads, OutputPadding);
+
+		TArray<int32> Result;
+		Result.SetNumUninitialized(XShape.Num());
+		Result[0] = XShape[0];
+		Result[1] = WShape[1] * Group;
+
+		int32 NumDimensions = XShape.Num() - 2;
+		for (int32 DimensionIndex = 0; DimensionIndex < NumDimensions; DimensionIndex++)
+		{
+			Result[DimensionIndex + 2] = Strides[DimensionIndex] * (XShape[DimensionIndex + 2] - 1) + OutputPadding[DimensionIndex] +
+				((WShape[DimensionIndex + 2] - 1) * Dilations[DimensionIndex] + 1) -
+				Padding[DimensionIndex] - Padding[Strides.Num() + DimensionIndex];
+		}
+
+		return Result;
+	}
+
+	void FConvTransposeCS::FillInParameters(EConvTransposeGroupSize GroupSize, TArrayView<const uint32> XShape, TArrayView<const uint32> WShape, bool HasB, EConvTransposeAutoPad AutoPad, int32 Group, TArrayView<const int32> Dilations, TArrayView<const int32> Strides, TArrayView<const int32> Pads, TArrayView<const int32> OutputPadding, FConvTransposeCS::FParameters& Parameters)
+	{
+		check(XShape.Num() > 2);
+		check(WShape.Num() == XShape.Num());
+		check(Dilations.Num() == 0 || Dilations.Num() == WShape.Num() - 2);
+		check(Strides.Num() == 0 || Strides.Num() == WShape.Num() - 2);
+		check(AutoPad != EConvTransposeAutoPad::NOTSET || Pads.Num() == 2 * (WShape.Num() - 2));
+		check(GetNumReadsPerThread(GroupSize, WShape, Dilations, Strides) >= 0)
+			check(WShape[0] > 0)
+			check(WShape[1] > 0)
+
+			int32 NumDimensions = XShape.Num() - 2;
+
+		TArray<int32> Padding = GetPadding(WShape, AutoPad, Dilations, Strides, Pads, OutputPadding);
+		TArray<int32> GroupShape = GetGroupShape(GroupSize, NumDimensions);
+		TArray<int32> YShape = GetOutputShape(XShape, WShape, AutoPad, Dilations, Strides, Pads, OutputPadding, Group);
+		TArray<int32> GridShape = GetGridShape(YShape, GroupShape);
+		TArray<int32> XBlockShape = GetXBlockShape(GroupShape, WShape, Dilations, Strides);
+
+		int32 GroupStride = 1;
+		int32 GroupThreadStride = 1;
+		int32 XBlockSize = 1;
+		int32 YMemoryStride = 1;
+		int32 XMemoryStride = 1;
+		int32 WChannelSize = 1;
+		for (int32 i = NumDimensions - 1; i >= 0; i--)
+		{
+			int32 Stride = i < Strides.Num() ? Strides[i] : 1;
+			int32 Dilation = i < Dilations.Num() ? Dilations[i] : 1;
+			Parameters.Dilation_Stride_XBlockStartOffset_DilationXBlockStride[i] = FIntVector4(Dilation, Stride, Padding[i] + Dilations[i] * (1 - WShape[2 + i]), Dilation * XBlockSize);
+			Parameters.GroupStride_GroupShape_GroupThreadStride_StrideXBlockStride[i] = FIntVector4(GroupStride, GroupShape[i], GroupThreadStride, XBlockSize);
+			Parameters.YDimension_YMemoryStride_XDimension_XMemoryStride[i] = FIntVector4(YShape[2 + i], YMemoryStride, XShape[2 + i], XMemoryStride);
+			Parameters.XBlockStartStride_XBlockStride_WDimension_WDimensionDilationXBlockStride[i] = FIntVector4(GroupShape[i], XBlockSize, WShape[2 + i], WShape[2 + i] * Dilation * XBlockSize);
+			Parameters.OneDiv_GroupStride_GroupThreadStride_OneDivStride[i] = FVector4f(1.0 / ((float)GroupStride), 1.0 / ((float)GroupThreadStride), 1.0 / ((float)XBlockSize), 1.0 / ((float)Stride));
+
+			GroupStride *= GridShape[i];
+			GroupThreadStride *= GroupShape[i];
+			XBlockSize *= XBlockShape[i];
+			YMemoryStride *= YShape[2 + i];
+			XMemoryStride *= XShape[2 + i];
+			WChannelSize *= WShape[2 + i];
+		}
+
+		Parameters.NumWChannels = WShape[0];
+		Parameters.NumOutChannelsDivGroup = WShape[1];
+
+		Parameters.YBatchStride = YShape[1] * YMemoryStride;
+		Parameters.YOutputKernelStride = YMemoryStride;
+
+		Parameters.XBatchStride = XShape[1] * XMemoryStride;
+		Parameters.XChannelStride = XMemoryStride;
+
+		Parameters.XBlockSize = XBlockSize;
+
+		Parameters.NumChannelsPerBatch = FMath::Min((int32)((float)GroupThreadStride / (float)WChannelSize), (int32)WShape[0]);
+		check(Parameters.NumChannelsPerBatch > 0)
+			Parameters.NumChannelBatches = FMath::DivideAndRoundUp((int32)WShape[0], Parameters.NumChannelsPerBatch);
+
+		Parameters.WOutputKernelStride = WShape[1] * WChannelSize;
+		Parameters.WChannelBatchSize = Parameters.NumChannelsPerBatch * WShape[1] * WChannelSize;
+		Parameters.WChannelSize = WChannelSize;
+
+		Parameters.GroupsDivM = 1.0 / ((float)WShape[1] * Group);
+		Parameters.OneDivGroup = 1.0 / (float)Group;
+	}
+
+	int32 FConvTransposeCS::GetNumReadsPerThread(EConvTransposeGroupSize GroupSize, TArrayView<const uint32> WShape, TArrayView<const int32> Dilations, TArrayView<const int32> Strides)
+	{
+		check(WShape.Num() > 2);
+		check(Dilations.Num() == 0 || Dilations.Num() == WShape.Num() - 2);
+		check(Strides.Num() == 0 || Strides.Num() == WShape.Num() - 2);
+
+		int32 NumDimensions = WShape.Num() - 2;
+
+		TArray<int32> GroupShape = GetGroupShape(GroupSize, NumDimensions);
+		int32 NumThreadsPerGroup = 1;
+		for (int32 i = 0; i < NumDimensions; i++)
+		{
+			NumThreadsPerGroup *= GroupShape[i];
+		}
+
+		TArray<int32> XBlockShape = GetXBlockShape(GroupShape, WShape, Dilations, Strides);
+		int32 NumXBlockElements = 1;
+		for (int32 i = 0; i < NumDimensions; i++)
+		{
+			NumXBlockElements *= XBlockShape[i];
+		}
+
+		int32 NumReads = FMath::DivideAndRoundUp(NumXBlockElements, NumThreadsPerGroup);
+		int32 NumReadsPow2 = FMath::Max(FMath::RoundToPositiveInfinity(FMath::Log2((float)NumReads)), FConvTransposeConstants::MIN_NUM_READS_PER_THREAD_POW2);
+
+		if (NumReadsPow2 <= FConvTransposeConstants::MAX_NUM_READS_PER_THREAD_POW2)
+		{
+			return NumReadsPow2;
+		}
+
+		return -1;
+	}
+
+	TArray<int32> FConvTransposeCS::GetGroupShape(EConvTransposeGroupSize GroupSize, int32 NumDimensions)
+	{
+		check(NumDimensions > 0);
+
+		int32 NumThreadsPerGroup = GetNumThreadsPerGroup(GroupSize);
+
+		int32 Power = (int32)FMath::Log2((float)NumThreadsPerGroup);
+		int32 MinPowerPerDim = (int32)((float)Power / (float)NumDimensions);
+		int32 PowerReminder = Power - NumDimensions * MinPowerPerDim;
+
+		TArray<int32> Result;
+		Result.Init((int32)FMath::Pow((float)2.0, (float)MinPowerPerDim), NumDimensions);
+		for (int32 i = 0; i < PowerReminder; i++)
+		{
+			Result[NumDimensions - 1 - i] *= 2;
+		}
+
+		return Result;
+	}
+
+	FIntVector FConvTransposeCS::GetGroupCount(TArrayView<const int32> YShape, TArrayView<const int32> GroupShape)
+	{
+		check(YShape.Num() > 2);
+		check(YShape.Num() == (GroupShape.Num() + 2));
+
+		int32 ThreadGroupCountValueX = 1;
+		for (int32 i = 2; i < YShape.Num(); i++)
+		{
+			ThreadGroupCountValueX *= FMath::DivideAndRoundUp(YShape[i], GroupShape[i - 2]);
+		}
+
+		return FIntVector(ThreadGroupCountValueX, YShape[1], YShape[0]);
+	}
+
+	EConvTransposeGroupSize FConvTransposeCS::GetMinimalGroupSize(TArrayView<const uint32> WShape)
+	{
+		int32 NumDimensions = WShape.Num() - 2;
+		int32 WChannelSize = 1;
+		for (int32 i = 0; i < NumDimensions; i++)
+		{
+			WChannelSize *= WShape[2 + i];
+		}
+
+		for (int32 i = 0; i < (int32)EConvTransposeGroupSize::MAX; i++)
+		{
+			if (GetNumThreadsPerGroup((EConvTransposeGroupSize)i) >= WChannelSize)
+			{
+				return (EConvTransposeGroupSize)i;
+			}
+		}
+
+		return EConvTransposeGroupSize::MAX;
+	}
+
+	TArray<int32> FConvTransposeCS::GetXBlockShape(TArrayView<const int32> GroupShape, TArrayView<const uint32> WShape, TArrayView<const int32> Dilations, TArrayView<const int32> Strides)
+	{
+		check(WShape.Num() > 2);
+		check(GroupShape.Num() == WShape.Num() - 2);
+		check(Dilations.Num() == 0 || Dilations.Num() == GroupShape.Num());
+		check(Strides.Num() == 0 || Strides.Num() == GroupShape.Num());
+
+		TArray<int32> Result;
+		Result.SetNumUninitialized(GroupShape.Num());
+		for (int32 i = 0; i < GroupShape.Num(); i++)
+		{
+			int32 DilatedKernelSize = (i < Dilations.Num() ? Dilations[i] : 1) * (WShape[2 + i] - 1) + 1;
+			Result[i] = DilatedKernelSize + (GroupShape[i] - 1);
+		}
+		return Result;
+	}
+
+	TArray<int32> FConvTransposeCS::GetPadding(TArrayView<const uint32> WShape, EConvTransposeAutoPad AutoPad, TArrayView<const int32> Dilations, TArrayView<const int32> Strides, TArrayView<const int32> Pads, TArrayView<const int32> OutputPadding)
+	{
+		check(WShape.Num() > 2);
+		check(Dilations.Num() == 0 || Dilations.Num() == WShape.Num() - 2);
+		check(Strides.Num() == 0 || Strides.Num() == WShape.Num() - 2);
+		check(AutoPad != EConvTransposeAutoPad::NOTSET || Pads.Num() == 2 * (WShape.Num() - 2));
+
+		int32 NumDimensions = WShape.Num() - 2;
+		TArray<int32> Result;
+		Result.Init(0, 2 * NumDimensions);
+
+		if (AutoPad == EConvTransposeAutoPad::NOTSET)
+		{
+			return TArray<int32>{Pads};
+		}
+		else if (AutoPad == EConvTransposeAutoPad::VALID)
+		{
+			return Result;
+		}
+
+		for (int32 DimensionIndex = 0; DimensionIndex < NumDimensions; DimensionIndex++)
+		{
+			int32 TotalPad = (WShape[DimensionIndex + 2] - 1) * Dilations[DimensionIndex] - Strides[DimensionIndex] + OutputPadding[DimensionIndex] + 1;
+
+			if (AutoPad == EConvTransposeAutoPad::SAME_LOWER)
+			{
+				Result[DimensionIndex] = (TotalPad + 1) / 2;
+			}
+			else
+			{
+				Result[DimensionIndex] = TotalPad / 2;
+			}
+			Result[NumDimensions + DimensionIndex] = TotalPad - Result[DimensionIndex];
+		}
+
+		return Result;
+	}
+
+	int32 FConvTransposeCS::GetNumThreadsPerGroup(EConvTransposeGroupSize GroupSize)
+	{
+		int32 NumThreadsPerGroup = 128;
+		switch (GroupSize)
+		{
+		case EConvTransposeGroupSize::Size128:
+			NumThreadsPerGroup = 128;
+			break;
+		case EConvTransposeGroupSize::Size256:
+			NumThreadsPerGroup = 256;
+			break;
+		case EConvTransposeGroupSize::Size512:
+			NumThreadsPerGroup = 512;
+			break;
+		default:
+			NumThreadsPerGroup = 128;
+			break;
+		}
+		check(FMath::Log2((float)NumThreadsPerGroup) == FMath::Floor(FMath::Log2((float)NumThreadsPerGroup)));
+		return NumThreadsPerGroup;
+	}
+
+	TArray<int32> FConvTransposeCS::GetGridShape(TArrayView<const int32> YShape, TArrayView<const int32> GroupShape)
+	{
+		check(YShape.Num() > 2);
+		check(YShape.Num() == (GroupShape.Num() + 2));
+
+		TArray<int32> Result;
+		Result.SetNumUninitialized(YShape.Num() - 2);
+		for (int32 i = 2; i < YShape.Num(); i++)
+		{
+			Result[i - 2] = FMath::DivideAndRoundUp(YShape[i], GroupShape[i - 2]);
+		}
+
+		return Result;
+	}
+
+	IMPLEMENT_GLOBAL_SHADER(FConvTransposeCS, "/NNI/NNIHlslShadersConvTranspose.usf", "ConvTranspose", SF_Compute);
+} // UE::NNIHlslShaders::Internal
\ No newline at end of file
diff --git a/Engine/Plugins/Experimental/NNX/Source/NNIHlslShaders/Private/NNIHlslShadersElementWiseBinaryCS.cpp b/Engine/Plugins/Experimental/NNX/Source/NNIHlslShaders/Private/NNIHlslShadersElementWiseBinaryCS.cpp
new file mode 100644
--- /dev/null
+++ b/Engine/Plugins/Experimental/NNX/Source/NNIHlslShaders/Private/NNIHlslShadersElementWiseBinaryCS.cpp
@@ -0,0 +1,59 @@
+// Copyright Epic Games, Inc. All Rights Reserved.
+
+#include "NNIHlslShadersElementWiseBinaryCS.h"
+#include "NNXCore.h"
+
+namespace UE::NNIHlslShaders::Internal
+{
+	void TElementWiseBinaryCS::ModifyCompilationEnvironment(const FGlobalShaderPermutationParameters& InParameters, FShaderCompilerEnvironment& OutEnvironment)
+	{
+		FGlobalShader::ModifyCompilationEnvironment(InParameters, OutEnvironment);
+
+		OutEnvironment.SetDefine(TEXT("THREADGROUP_SIZE_X"), FElementWiseBinaryConstants::NUM_GROUP_THREADS);
+
+		FPermutationDomain PermutationVector(InParameters.PermutationId);
+
+		const FString OpFunc = GetOpFunc(PermutationVector.Get<FOperatorType>());
+
+		OutEnvironment.SetDefine(TEXT("ELEMENTWISE_OP(X,Y)"), *OpFunc);
+	}
+
+	const FString TElementWiseBinaryCS::GetOpFunc(EMLElementWiseBinaryOperatorType OpType)
+	{
+		FString OpTable[(int32) EMLElementWiseBinaryOperatorType::MAX];
+
+		for (int32 Idx = 0; Idx < (int32)EMLElementWiseBinaryOperatorType::MAX; ++Idx)
+		{
+			OpTable[Idx] = FString("");
+		}
+
+#define OP(OpName, OpFunc) OpTable[(int32) EMLElementWiseBinaryOperatorType::OpName] = OpFunc
+		OP(Add,              TEXT("((X)+(Y))"));
+		//OP(And,            TEXT("((X)&&(Y))"));
+		OP(Div,              TEXT("((X)/(Y))"));
+		//OP(Equal,          TEXT("((X)==(Y))"));
+		//OP(Greater,        TEXT("((X)>(Y))"));
+		//OP(GreaterOrEqual, TEXT("((X)>=(Y))"));
+		//OP(Less,           TEXT("((X)<(Y))"));
+		//OP(LessOrEqual,    TEXT("((X)<(Y))"));
+		OP(Mod,              TEXT("((X)%(Y))"));
+		OP(Mul,              TEXT("((X)*(Y))"));
+		//OP(Or,             TEXT("((X)||(Y))"));
+		OP(Prelu,            TEXT("prelu(X,Y)"));
+		OP(Pow,              TEXT("safe_pow(X,Y)"));
+		OP(Sub,              TEXT("((X)-(Y))"));
+		//OP(Or,             TEXT("((X)^=(Y))"));
+#undef OP
+
+		FString OpFunc = OpTable[(int32) OpType];
+
+		if (OpFunc == "")
+		{
+			UE_LOG(LogNNX, Warning, TEXT("Undefined ElementWise Binary operator name for operator:%d"), OpType);
+		}
+
+		return OpFunc;
+	}
+
+	IMPLEMENT_GLOBAL_SHADER(TElementWiseBinaryCS, "/NNI/NNIHlslShadersElementWiseBinary.usf", "ElementWiseBinary", SF_Compute);
+} // UE::NNIHlslShaders::Internal
\ No newline at end of file
diff --git a/Engine/Plugins/Experimental/NNX/Source/NNIHlslShaders/Private/NNIHlslShadersElementWiseUnaryCS.cpp b/Engine/Plugins/Experimental/NNX/Source/NNIHlslShaders/Private/NNIHlslShadersElementWiseUnaryCS.cpp
new file mode 100644
--- /dev/null
+++ b/Engine/Plugins/Experimental/NNX/Source/NNIHlslShaders/Private/NNIHlslShadersElementWiseUnaryCS.cpp
@@ -0,0 +1,82 @@
+// Copyright Epic Games, Inc. All Rights Reserved.
+
+#include "NNIHlslShadersElementWiseUnaryCS.h"
+#include "NNXCore.h"
+
+namespace UE::NNIHlslShaders::Internal
+{
+	void TElementWiseUnaryCS::ModifyCompilationEnvironment(const FGlobalShaderPermutationParameters& InParameters, FShaderCompilerEnvironment& OutEnvironment)
+	{
+		FGlobalShader::ModifyCompilationEnvironment(InParameters, OutEnvironment);
+
+		OutEnvironment.SetDefine(TEXT("THREADGROUP_SIZE_X"), FElementWiseUnaryConstants::NUM_GROUP_THREADS);
+
+		FPermutationDomain PermutationVector(InParameters.PermutationId);
+
+		const FString OpFunc = GetOpFunc(PermutationVector.Get<FOperatorType>());
+
+		OutEnvironment.SetDefine(TEXT("ELEMENTWISE_OP(X)"), *OpFunc);
+	}
+
+	const FString TElementWiseUnaryCS::GetOpFunc(EMLElementWiseUnaryOperatorType OpType)
+	{
+		FString OpTable[(int32) EMLElementWiseUnaryOperatorType::MAX];
+
+		for (int32 Idx = 0; Idx < (int32) EMLElementWiseUnaryOperatorType::MAX; ++Idx)
+		{
+			OpTable[Idx] = FString("");
+		}
+
+#define OP(OpName, OpFunc) OpTable[(int32) EMLElementWiseUnaryOperatorType::OpName] = OpFunc
+		OP(Abs,         TEXT("abs(X)"));
+		OP(Acos,        TEXT("acos(X)"));
+		OP(Acosh,       TEXT("acosh(X)"));
+		OP(Asin,        TEXT("asin(X)"));
+		OP(Asinh,       TEXT("asinh(X)"));
+		OP(Atan,        TEXT("atan(X)"));
+		OP(Atanh,       TEXT("atanh(X)"));
+		//OP(BitShift,  TEXT("bitshift(X)"));
+		//OP(Cast,      TEXT("cast(X)"));
+		OP(Ceil,        TEXT("ceil(X)"));
+		//OP(Clip,      TEXT("clip(X)"));
+		OP(Cos,         TEXT("cos(X)"));
+		OP(Cosh,        TEXT("cosh(X)"));
+		OP(Elu,         TEXT("elu(X)"));
+		OP(Erf,         TEXT("erf(X)"));
+		OP(Exp,         TEXT("exp(X)"));
+		OP(Floor,       TEXT("floor(X)"));
+		OP(IsInf,       TEXT("isinf(X)"));
+		OP(IsNan,       TEXT("isnan(X)"));//TODO check shader there is a warning saying input can neither be Nan on PC FXC
+		OP(HardSigmoid, TEXT("hardSigmoid(X)"));
+		OP(HardSwish,   TEXT("hardSwish(X)"));
+		OP(LeakyRelu,   TEXT("leakyRelu(X)"));
+		OP(Log,         TEXT("log(X)"));
+		OP(Neg,         TEXT("-(X)"));
+		//OP(Not,       TEXT("not(X)"));
+		OP(Reciprocal,  TEXT("1 / (X)"));
+		OP(Relu,        TEXT("relu(X)"));
+		OP(Round,       TEXT("round(X)"));
+		OP(Selu,        TEXT("selu(X)"));
+		OP(Sigmoid,     TEXT("sigmoid(X)"));
+		OP(Sign,        TEXT("sign(X)"));
+		OP(Sin,         TEXT("sin(X)"));
+		OP(Sinh,        TEXT("sinh(X)"));
+		OP(Softplus,    TEXT("softplus(X)"));
+		OP(Softsign,    TEXT("softsign(X)"));
+		OP(Sqrt,        TEXT("sqrt(X)"));
+		OP(Tan,         TEXT("tan(X)"));
+		OP(Tanh,        TEXT("tanh(X)"));
+#undef OP
+
+		FString OpFunc = OpTable[(int32) OpType];
+
+		if (OpFunc == "")
+		{
+			UE_LOG(LogNNX, Warning, TEXT("Undefined ElementWise Unary operator name for operator:%d"), OpType);
+		}
+
+		return OpFunc;
+	}
+
+	IMPLEMENT_GLOBAL_SHADER(TElementWiseUnaryCS, "/NNI/NNIHlslShadersElementWise.usf", "ElementWise", SF_Compute);
+} // UE::NNIHlslShaders::Internal
diff --git a/Engine/Plugins/Experimental/NNX/Source/NNIHlslShaders/Private/NNIHlslShadersElementWiseVariadicCS.cpp b/Engine/Plugins/Experimental/NNX/Source/NNIHlslShaders/Private/NNIHlslShadersElementWiseVariadicCS.cpp
new file mode 100644
--- /dev/null
+++ b/Engine/Plugins/Experimental/NNX/Source/NNIHlslShaders/Private/NNIHlslShadersElementWiseVariadicCS.cpp
@@ -0,0 +1,48 @@
+// Copyright Epic Games, Inc. All Rights Reserved.
+
+#include "NNIHlslShadersElementWiseVariadicCS.h"
+#include "NNXCore.h"
+
+namespace UE::NNIHlslShaders::Internal
+{
+	void TElementWiseVariadicCS::ModifyCompilationEnvironment(const FGlobalShaderPermutationParameters& InParameters, FShaderCompilerEnvironment& OutEnvironment)
+	{
+		FGlobalShader::ModifyCompilationEnvironment(InParameters, OutEnvironment);
+
+		OutEnvironment.SetDefine(TEXT("THREADGROUP_SIZE_X"), FElementWiseVariadicConstants::NUM_GROUP_THREADS);
+
+		FPermutationDomain PermutationVector(InParameters.PermutationId);
+
+		const FString OpFunc = GetOpFunc(PermutationVector.Get<FOperatorType>());
+
+		OutEnvironment.SetDefine(TEXT("ELEMENTWISE_OP(X,Y)"), *OpFunc);
+	}
+
+	const FString TElementWiseVariadicCS::GetOpFunc(EMLElementWiseVariadicOperatorType OpType)
+	{
+		FString OpTable[(int32)EMLElementWiseVariadicOperatorType::MAX];
+
+		for (int32 Idx = 0; Idx < (int32)EMLElementWiseVariadicOperatorType::MAX; ++Idx)
+		{
+			OpTable[Idx] = FString("");
+		}
+
+#define OP(OpName, OpFunc) OpTable[(int32) EMLElementWiseVariadicOperatorType::OpName] = OpFunc
+		OP(Max, TEXT("max(X,Y)"));
+		OP(Min, TEXT("min(X,Y)"));
+		OP(Mean, TEXT("((X)+(Y))"));
+		OP(Sum, TEXT("((X)+(Y))"));
+#undef OP
+
+		FString OpFunc = OpTable[(int32)OpType];
+
+		if (OpFunc == "")
+		{
+			UE_LOG(LogNNX, Warning, TEXT("Undefined ElementWise Variadic operator name for operator:%d"), OpType);
+		}
+
+		return OpFunc;
+	}
+
+	IMPLEMENT_GLOBAL_SHADER(TElementWiseVariadicCS, "/NNI/NNIHlslShadersElementWiseVariadic.usf", "ElementWiseVariadic", SF_Compute);
+} // UE::NNIHlslShaders::Internal
\ No newline at end of file
diff --git a/Engine/Plugins/Experimental/NNX/Source/NNIHlslShaders/Private/NNIHlslShadersGatherCS.cpp b/Engine/Plugins/Experimental/NNX/Source/NNIHlslShaders/Private/NNIHlslShadersGatherCS.cpp
new file mode 100644
--- /dev/null
+++ b/Engine/Plugins/Experimental/NNX/Source/NNIHlslShaders/Private/NNIHlslShadersGatherCS.cpp
@@ -0,0 +1,81 @@
+// Copyright Epic Games, Inc. All Rights Reserved.
+
+#include "NNIHlslShadersGatherCS.h"
+
+namespace UE::NNIHlslShaders::Internal
+{
+	// template <>
+	void TGatherCS::ModifyCompilationEnvironment(const FGlobalShaderPermutationParameters& InParameters, FShaderCompilerEnvironment& OutEnvironment)
+	{
+		FGlobalShader::ModifyCompilationEnvironment(InParameters, OutEnvironment);
+		OutEnvironment.SetDefine(TEXT("DATA_TYPE"), 0);
+		OutEnvironment.SetDefine(TEXT("INDICES_TYPE"), 0);
+		OutEnvironment.SetDefine(TEXT("MAX_NUM_DIMENSIONS"), FGatherConstants::MAX_NUM_DIMENSIONS);
+		OutEnvironment.SetDefine(TEXT("NUM_GROUP_THREADS"), FGatherConstants::NUM_GROUP_THREADS);
+	}
+
+	//template <typename DataElementType, typename IndicesElementType>
+	void TGatherCS::FillInParameters(int32 Axis, const NNX::FMLTensorDesc& Data, const NNX::FMLTensorDesc& Indices, FParameters& Parameters)
+	{
+		Parameters.Axis = Axis;
+
+		TArray<int32> OutputShape;
+		OutputShape.SetNum(Indices.Shape.Num() + (Data.Shape.Num() - 1));
+		int32 index = 0;
+		Parameters.OutputSize = 1;
+		for (int32 i = 0; i < Axis; i++, index++)
+		{
+			OutputShape[index] = Data.Shape[i];
+			Parameters.OutputSize *= OutputShape[index];
+		}
+		for (int32 i = 0; i < Indices.Shape.Num(); i++, index++)
+		{
+			OutputShape[index] = Indices.Shape[i];
+			Parameters.OutputSize *= OutputShape[index];
+		}
+		for (int32 i = (Axis + 1); i < Data.Shape.Num(); i++, index++)
+		{
+			OutputShape[index] = Data.Shape[i];
+			Parameters.OutputSize *= OutputShape[index];
+		}
+
+		Parameters.NumDataDimensions = Data.Shape.Num();
+		Parameters.NumIndicesDimensions = Indices.Shape.Num();
+
+		int32 DataStride = 1;
+		for (int32 i = Data.Shape.Num() - 1; i >= 0; i--)
+		{
+			Parameters.DataStride_IndicesStride_OutputStride[i].X = DataStride;
+			Parameters.OneDivDataStride_OneDivIndicesStride_OneDivOutputStride[i].X = 1.0 / (float)DataStride;
+			DataStride *= Data.Shape[i];
+		}
+
+		int32 IndicesStride = 1;
+		for (int32 i = Indices.Shape.Num() - 1; i >= 0; i--)
+		{
+			Parameters.DataStride_IndicesStride_OutputStride[i].Y = IndicesStride;
+			Parameters.OneDivDataStride_OneDivIndicesStride_OneDivOutputStride[i].Y = 1.0 / (float)IndicesStride;
+			IndicesStride *= Indices.Shape[i];
+		}
+
+		int32 OutputStride = 1;
+		for (int32 i = OutputShape.Num() - 1; i >= 0; i--)
+		{
+			Parameters.DataStride_IndicesStride_OutputStride[i].Z = OutputStride;
+			Parameters.OneDivDataStride_OneDivIndicesStride_OneDivOutputStride[i].Z = 1.0 / (float)OutputStride;
+			OutputStride *= OutputShape[i];
+		}
+	}
+
+	// template <typename DataElementType, typename IndicesElementType>
+	FIntVector TGatherCS::GetGroupCount(const FParameters& Parameters)
+	{
+		return FIntVector(FMath::DivideAndRoundUp(Parameters.OutputSize, FGatherConstants::NUM_GROUP_THREADS), 1, 1);
+	}
+
+	//typedef TGatherCS<float, int32> FGatherFloatInt32CS;
+	//IMPLEMENT_SHADER_TYPE(template<>, TGatherCS, TEXT("/NNX/GatherOp.usf"), TEXT("main"), SF_Compute);
+	IMPLEMENT_GLOBAL_SHADER(TGatherCS, "/NNI/NNIHlslShadersGather.usf", "Gather", SF_Compute);
+	//template class TGatherCS<float, int32>;
+
+} // UE::NNIHlslShaders::Internal
\ No newline at end of file
diff --git a/Engine/Plugins/Experimental/NNX/Source/NNIHlslShaders/Private/NNIHlslShadersGemmCS.cpp b/Engine/Plugins/Experimental/NNX/Source/NNIHlslShaders/Private/NNIHlslShadersGemmCS.cpp
new file mode 100644
--- /dev/null
+++ b/Engine/Plugins/Experimental/NNX/Source/NNIHlslShaders/Private/NNIHlslShadersGemmCS.cpp
@@ -0,0 +1,190 @@
+// Copyright Epic Games, Inc. All Rights Reserved.
+
+#include "NNIHlslShadersGemmCS.h"
+#include "NNXCore.h"
+#include "Algo/Accumulate.h"
+
+namespace UE::NNIHlslShaders::Internal
+{
+	struct GemmMatrixParameters
+	{
+		uint32 M = 0;
+		uint32 N = 0;
+		uint32 K = 0;
+		TArray<uint32> StackShapeA;
+		TArray<uint32> StackShapeB;
+		TArray<uint32> StackStrideA;
+		TArray<uint32> StackStrideB;
+	};
+
+	static GemmMatrixParameters GetMatrixParametersMatMul(const NNX::FMLTensorDesc& InputA, const NNX::FMLTensorDesc& InputB)
+	{
+		check(InputA.Shape.Num() != 0);
+		check(InputB.Shape.Num() != 0);
+
+		GemmMatrixParameters Result;
+		int32 NumStackDimensionsA = FMath::Max((int32)InputA.Shape.Num() - 2, 0);
+		int32 NumStackDimensionsB = FMath::Max((int32)InputB.Shape.Num() - 2, 0);
+		int32 NumStackDimensions = FMath::Max(NumStackDimensionsA, NumStackDimensionsB);
+
+		check(NumStackDimensions <= FGemmConstants::MAX_NUM_STACK_DIMENSIONS);
+
+		// Check matrix stack dimensions
+		if (NumStackDimensionsA > 0 && NumStackDimensionsB > 0)
+		{
+			const int32 NumDimToCheck = FMath::Min(NumStackDimensionsA, NumStackDimensionsB);
+			for (int32 i = 0; i < NumDimToCheck; i++)
+			{
+				const uint32 VolumeA = InputA.Shape[NumStackDimensionsA - 1 - i];
+				const uint32 VolumeB = InputB.Shape[NumStackDimensionsB - 1 - i];
+
+				check(VolumeA == 1 || VolumeB == 1 || VolumeA == VolumeB);
+			}
+		}
+
+		const bool IsVectorA = InputA.Shape.Num() == 1;
+		const bool IsVectorB = InputB.Shape.Num() == 1;
+
+		Result.M = IsVectorA ? 1 : InputA.Shape[InputA.Shape.Num() - 2];
+		Result.N = IsVectorB ? 1 : InputB.Shape[InputB.Shape.Num() - 1];
+		Result.K = InputA.Shape[IsVectorA ? 0 : InputA.Shape.Num() - 1];
+		check(InputB.Shape[IsVectorB ? 0 : InputB.Shape.Num() - 2] == Result.K);
+
+		Result.StackShapeA.Init(1, NumStackDimensions);
+		Result.StackShapeB.Init(1, NumStackDimensions);
+
+		for (int32 i = 0; i < NumStackDimensions; i++)
+		{
+			int32 IdxA = InputA.Shape.Num() - 3 - i;
+			int32 IdxB = InputB.Shape.Num() - 3 - i;
+
+			Result.StackShapeA[Result.StackShapeA.Num() - 1 - i] = IdxA >= 0 ? InputA.Shape[IdxA] : 1;
+			Result.StackShapeB[Result.StackShapeB.Num() - 1 - i] = IdxB >= 0 ? InputB.Shape[IdxB] : 1;
+		}
+
+		Result.StackStrideA.Init(1, NumStackDimensions);
+		Result.StackStrideB.Init(1, NumStackDimensions);
+
+		for (int32 i = NumStackDimensions - 2; i >= 0; i--)
+		{
+			Result.StackStrideA[i] = Result.StackStrideA[i + 1] * Result.StackShapeA[i + 1];
+			Result.StackStrideB[i] = Result.StackStrideB[i + 1] * Result.StackShapeB[i + 1];
+		}
+
+		return Result;
+	}
+
+	void TGemmCS::ModifyCompilationEnvironment(const FGlobalShaderPermutationParameters& InParameters, FShaderCompilerEnvironment& OutEnvironment)
+	{
+		FGlobalShader::ModifyCompilationEnvironment(InParameters, OutEnvironment);
+		OutEnvironment.SetDefine(TEXT("MAX_NUM_STACK_DIMENSIONS"), FGemmConstants::MAX_NUM_STACK_DIMENSIONS);
+	}
+
+	void TGemmCS::FillInParameters(float Alpha, float Beta, int32 TransA, int32 TransB, const NNX::FMLTensorDesc& InputA, const NNX::FMLTensorDesc& InputB, const NNX::FMLTensorDesc& InputC, float CScalar, TGemmCS::FParameters& Parameters)
+	{
+		check(InputA.Shape.Num() == 2);
+		check(InputB.Shape.Num() == 2);
+
+		uint32 M = TransA != 0 ? InputA.Shape[1] : InputA.Shape[0];
+		uint32 K = TransA != 0 ? InputA.Shape[0] : InputA.Shape[1];
+		uint32 N = TransB != 0 ? InputB.Shape[0] : InputB.Shape[1];
+		check(K == (TransB != 0 ? InputB.Shape[1] : InputB.Shape[0]));
+
+		Parameters.Alpha = Alpha;
+		Parameters.Beta = Beta;
+		Parameters.TransA = TransA;
+		Parameters.TransB = TransB;
+		Parameters.M = M;
+		Parameters.N = N;
+		Parameters.K = K;
+		Parameters.MxK = M * K;
+		Parameters.KxN = K * N;
+		Parameters.MxN = M * N;
+		Parameters.CWidth = InputC.Shape.Num() == 0 ? 0 : InputC.Shape[InputC.Shape.Num() - 1];;
+		Parameters.CHeight = InputC.Shape.Num() == 0 ? 0 : InputC.Shape.Num() == 1 ? 1 : InputC.Shape[InputC.Shape.Num() - 2];
+		Parameters.CScalar = CScalar;
+	}
+
+	void TGemmCS::FillInParametersMatMul(const NNX::FMLTensorDesc& InputA, const NNX::FMLTensorDesc& InputB, TGemmCS::FParameters& Parameters)
+	{
+		GemmMatrixParameters Params = GetMatrixParametersMatMul(InputA, InputB);
+
+		Parameters.Alpha = 1.0f;
+		Parameters.Beta = 1.0f;
+		Parameters.TransA = 0;
+		Parameters.TransB = 0;
+		Parameters.M = Params.M;
+		Parameters.N = Params.N;
+		Parameters.K = Params.K;
+		Parameters.MxK = Params.M * Params.K;
+		Parameters.KxN = Params.K * Params.N;
+		Parameters.MxN = Params.M * Params.N;
+		Parameters.CWidth = 0;
+		Parameters.CHeight = 0;
+		Parameters.CScalar = 0.0f;
+
+		for (int32 i = 0; i < Params.StackShapeA.Num(); i++)
+		{
+			Parameters.StackShapeA_StackShapeB_StackStrideA_StackStrideB[i] = FUint32Vector4(Params.StackShapeA[i], Params.StackShapeB[i], Params.StackStrideA[i], Params.StackStrideB[i]);
+		}
+	}
+
+	FIntVector TGemmCS::GetGroupCount(const TGemmCS::FParameters& Parameters, EGemmAlgorithm Algorithm, int32 NumStackDimensions)
+	{
+		int32 OutputStackSize = 1;
+		for (int32 i = 0; i < NumStackDimensions; i++)
+		{
+			OutputStackSize *= FMath::Max(
+				Parameters.StackShapeA_StackShapeB_StackStrideA_StackStrideB[i].X,
+				Parameters.StackShapeA_StackShapeB_StackStrideA_StackStrideB[i].Y);
+		}
+
+		int32 ThreadGroupCountValueZ = FMath::DivideAndRoundUp(OutputStackSize, 1);
+		int32 ThreadGroupCountValueY = FMath::DivideAndRoundUp((int32)Parameters.M, 8);
+		int32 ThreadGroupCountValueX = FMath::DivideAndRoundUp((int32)Parameters.N, 8);
+		switch (Algorithm)
+		{
+		case EGemmAlgorithm::Simple8x8:
+		case EGemmAlgorithm::SharedMemory8x8:
+			ThreadGroupCountValueX = FMath::DivideAndRoundUp((int32)Parameters.N, 8);
+			ThreadGroupCountValueY = FMath::DivideAndRoundUp((int32)Parameters.M, 8);
+			break;
+		case EGemmAlgorithm::Simple16x16:
+		case EGemmAlgorithm::SharedMemory16x16:
+		case EGemmAlgorithm::MultiWrite1x16:
+		case EGemmAlgorithm::MultiWrite2x16:
+			ThreadGroupCountValueX = FMath::DivideAndRoundUp((int32)Parameters.N, 16);
+			ThreadGroupCountValueY = FMath::DivideAndRoundUp((int32)Parameters.M, 16);
+			break;
+		case EGemmAlgorithm::Simple32x32:
+		case EGemmAlgorithm::SharedMemory32x32:
+		case EGemmAlgorithm::MultiWrite1x32:
+		case EGemmAlgorithm::MultiWrite2x32:
+		case EGemmAlgorithm::MultiWrite4x32:
+			ThreadGroupCountValueX = FMath::DivideAndRoundUp((int32)Parameters.N, 32);
+			ThreadGroupCountValueY = FMath::DivideAndRoundUp((int32)Parameters.M, 32);
+			break;
+		case EGemmAlgorithm::Simple256x1:
+			ThreadGroupCountValueX = FMath::DivideAndRoundUp((int32)Parameters.N, 256);
+			ThreadGroupCountValueY = Parameters.M;
+			break;
+		case EGemmAlgorithm::MultiWrite2x64:
+		case EGemmAlgorithm::MultiWrite4x64:
+		case EGemmAlgorithm::MultiWrite8x64:
+			ThreadGroupCountValueX = FMath::DivideAndRoundUp((int32)Parameters.N, 64);
+			ThreadGroupCountValueY = FMath::DivideAndRoundUp((int32)Parameters.M, 64);
+			break;
+		default:
+			break;
+		}
+
+		return FIntVector(ThreadGroupCountValueX, ThreadGroupCountValueY, ThreadGroupCountValueZ);
+	}
+
+	EGemmAlgorithm TGemmCS::GetAlgorithm(const TGemmCS::FParameters& Parameters)
+	{
+		return EGemmAlgorithm::MultiWrite1x32;
+	}
+
+	IMPLEMENT_GLOBAL_SHADER(TGemmCS, "/NNI/NNIHlslShadersGemm.usf", "Gemm", SF_Compute);
+} // UE::NNIHlslShaders::Internal
\ No newline at end of file
diff --git a/Engine/Plugins/Experimental/NNX/Source/NNIHlslShaders/Private/NNIHlslShadersModule.cpp b/Engine/Plugins/Experimental/NNX/Source/NNIHlslShaders/Private/NNIHlslShadersModule.cpp
new file mode 100644
--- /dev/null
+++ b/Engine/Plugins/Experimental/NNX/Source/NNIHlslShaders/Private/NNIHlslShadersModule.cpp
@@ -0,0 +1,41 @@
+// Copyright Epic Games, Inc. All Rights Reserved.
+
+#include "NNXCore.h"
+#include "CoreMinimal.h"
+#include "Modules/ModuleManager.h"
+#include "Interfaces/IPluginManager.h"
+
+#include "Misc/Paths.h"
+#include "ShaderCore.h"
+
+class FNNIHlslShadersModule : public IModuleInterface
+{
+public:
+	// Begin IModuleInterface
+	virtual void StartupModule() override
+	{
+		FString BaseDir;
+
+		const TSharedPtr<IPlugin> Plugin = IPluginManager::Get().FindPlugin(TEXT("NNX"));
+		
+		if (Plugin.IsValid())
+		{
+			BaseDir = Plugin->GetBaseDir() + TEXT("/Source/NNIHlslShaders");
+		}
+		else
+		{
+			UE_LOG(LogNNX, Warning, TEXT("Shaders directory not added. Failed to find NNI plugin"));
+		}
+
+		FString ModuleShaderDir = FPaths::Combine(BaseDir, TEXT("Shaders"));
+		
+		AddShaderSourceDirectoryMapping(TEXT("/NNI"), ModuleShaderDir);
+	}
+
+	virtual void ShutdownModule() override
+	{
+
+	}
+};
+
+IMPLEMENT_MODULE(FNNIHlslShadersModule, NNIHlslShaders);
diff --git a/Engine/Plugins/Experimental/NNX/Source/NNIHlslShaders/Public/NNXShaderParameters.h b/Engine/Plugins/Experimental/NNX/Source/NNIHlslShaders/Public/NNXShaderParameters.h
new file mode 100644
--- /dev/null
+++ b/Engine/Plugins/Experimental/NNX/Source/NNIHlslShaders/Public/NNXShaderParameters.h
@@ -0,0 +1,174 @@
+// Copyright Epic Games, Inc. All Rights Reserved.
+
+#pragma once
+
+#include "ShaderParameterUtils.h"
+#include "RenderGraphUtils.h"
+
+// HLSL runtime and DirectML runtime use different shader parameter access patterns
+#define NNXRT_RDG_DML
+
+#ifdef NNXRT_RDG_DML
+#define NNXRT_RDG_BUFFER_SRV(View, BufferName) RDG_BUFFER_ACCESS(BufferName, ERHIAccess::UAVCompute)
+#define NNXRT_RDG_BUFFER_UAV(View, BufferName) RDG_BUFFER_ACCESS(BufferName, ERHIAccess::UAVCompute)
+#else
+#define NNXRT_RDG_BUFFER_SRV(View, BufferName) SHADER_PARAMETER_RDG_BUFFER_SRV(View, BufferName)
+#define NNXRT_RDG_BUFFER_UAV(View, BufferName) SHADER_PARAMETER_RDG_BUFFER_UAV(View, BufferName)
+#endif
+
+#ifdef NNXRT_RDG_DML
+#define NNXRT_BEGIN_SHADER_PARAMETER_STRUCT(Name) BEGIN_SHADER_PARAMETER_STRUCT(Name, )
+#else
+#define NNXRT_BEGIN_SHADER_PARAMETER_STRUCT(Name) BEGIN_SHADER_PARAMETER_STRUCT(FParameters, )
+#endif
+
+#define NNXRT_END_SHADER_PARAMETER_STRUCT END_SHADER_PARAMETER_STRUCT
+
+#define NXRT_TENSORSTRIDEINFO_MAX_NUM_DIMENSIONS 8
+
+//Unary element wise operators
+#define NNXRT_ELEMENTWISEUNARY_PARAMETER_STRUCT() \
+	NNXRT_BEGIN_SHADER_PARAMETER_STRUCT(FMLElementWiseUnaryParameters) \
+		NNXRT_RDG_BUFFER_SRV(Buffer<float>, Input) \
+		NNXRT_RDG_BUFFER_UAV(RWBuffer<float>, Output) \
+		SHADER_PARAMETER(uint32, Num) \
+		SHADER_PARAMETER(uint32, ThreadCountX) \
+		SHADER_PARAMETER(float, Alpha) \
+		SHADER_PARAMETER(float, Beta) \
+		SHADER_PARAMETER(float, Gamma) \
+	NNXRT_END_SHADER_PARAMETER_STRUCT()
+
+#ifdef NNXRT_RDG_DML
+NNXRT_ELEMENTWISEUNARY_PARAMETER_STRUCT()
+#endif
+
+//Variadic element wise operators
+#define NNXRT_ELEMENTWISEVARIADIC_PARAMETER_STRUCT() \
+	NNXRT_BEGIN_SHADER_PARAMETER_STRUCT(FMLElementWiseVariadicParameters) \
+		NNXRT_RDG_BUFFER_SRV(Buffer<float>, Input0) \
+		NNXRT_RDG_BUFFER_SRV(Buffer<float>, Input1) \
+		NNXRT_RDG_BUFFER_SRV(Buffer<float>, Input2) \
+		NNXRT_RDG_BUFFER_SRV(Buffer<float>, Input3) \
+		NNXRT_RDG_BUFFER_UAV(RWBuffer<float>, Output) \
+		SHADER_PARAMETER_ARRAY(FUintVector4, InputTensorInfo, [NXRT_TENSORSTRIDEINFO_MAX_NUM_DIMENSIONS]) \
+		SHADER_PARAMETER_ARRAY(FUintVector4, OutputTensorInfo, [NXRT_TENSORSTRIDEINFO_MAX_NUM_DIMENSIONS]) \
+		SHADER_PARAMETER(uint32, Num) \
+		SHADER_PARAMETER(uint32, ThreadCountX) \
+		SHADER_PARAMETER(float, Scale) \
+	NNXRT_END_SHADER_PARAMETER_STRUCT()
+
+#ifdef NNXRT_RDG_DML
+NNXRT_ELEMENTWISEVARIADIC_PARAMETER_STRUCT()
+#endif
+
+//Binary element wise operators
+#define NNXRT_ELEMENTWISEBINARY_PARAMETER_STRUCT() \
+	NNXRT_BEGIN_SHADER_PARAMETER_STRUCT(FMLElementWiseBinaryParameters) \
+		NNXRT_RDG_BUFFER_SRV(Buffer<float>, LHSInput) \
+		NNXRT_RDG_BUFFER_SRV(Buffer<float>, RHSInput) \
+		NNXRT_RDG_BUFFER_UAV(RWBuffer<float>, Output) \
+		SHADER_PARAMETER_ARRAY(FUintVector4, TensorInfo, [NXRT_TENSORSTRIDEINFO_MAX_NUM_DIMENSIONS]) \
+		SHADER_PARAMETER(uint32, Num) \
+		SHADER_PARAMETER(uint32, ThreadCountX) \
+	NNXRT_END_SHADER_PARAMETER_STRUCT()
+
+#ifdef NNXRT_RDG_DML
+NNXRT_ELEMENTWISEBINARY_PARAMETER_STRUCT()
+#endif
+
+//Gemm
+#define NNXRT_GEMM_MAX_NUM_STACK_DIMENSIONS 8
+
+#define NNXRT_GEMM_PARAMETER_STRUCT() \
+	NNXRT_BEGIN_SHADER_PARAMETER_STRUCT(FMLGemmParameters) \
+		SHADER_PARAMETER(float, Alpha) \
+		SHADER_PARAMETER(float, Beta) \
+		SHADER_PARAMETER(int32, TransA) \
+		SHADER_PARAMETER(int32, TransB) \
+		SHADER_PARAMETER(uint32, M) \
+		SHADER_PARAMETER(uint32, N) \
+		SHADER_PARAMETER(uint32, K) \
+		SHADER_PARAMETER(uint32, MxK) \
+		SHADER_PARAMETER(uint32, KxN) \
+		SHADER_PARAMETER(uint32, MxN) \
+		SHADER_PARAMETER(uint32, CWidth) \
+		SHADER_PARAMETER(uint32, CHeight) \
+		SHADER_PARAMETER(float, CScalar) \
+		SHADER_PARAMETER_ARRAY(FUint32Vector4, StackShapeA_StackShapeB_StackStrideA_StackStrideB, [NNXRT_GEMM_MAX_NUM_STACK_DIMENSIONS]) \
+		NNXRT_RDG_BUFFER_SRV(Buffer<float>, A) \
+		NNXRT_RDG_BUFFER_SRV(Buffer<float>, B) \
+		NNXRT_RDG_BUFFER_SRV(Buffer<float>, C) \
+		NNXRT_RDG_BUFFER_UAV(RWBuffer<float>, Y) \
+	NNXRT_END_SHADER_PARAMETER_STRUCT()
+
+#ifdef NNXRT_RDG_DML
+NNXRT_GEMM_PARAMETER_STRUCT()
+#endif
+
+//Conv
+#define NNXRT_CONV_MAX_NUM_DIMENSIONS 4
+
+#define NNXRT_CONV_PARAMETER_STRUCT() \
+	NNXRT_BEGIN_SHADER_PARAMETER_STRUCT(FMLConvParameters) \
+		NNXRT_RDG_BUFFER_SRV(Buffer<float>, X) \
+		NNXRT_RDG_BUFFER_SRV(Buffer<float>, W) \
+		NNXRT_RDG_BUFFER_UAV(RWBuffer<float>, Y) \
+		NNXRT_RDG_BUFFER_SRV(Buffer<float>, B) \
+		SHADER_PARAMETER_ARRAY(FIntVector4, Dilation_Stride_XBlockStartOffset_DilationXBlockStride, [NNXRT_CONV_MAX_NUM_DIMENSIONS]) \
+		SHADER_PARAMETER_ARRAY(FIntVector4, GroupStride_GroupShape_GroupThreadStride_StrideXBlockStride, [NNXRT_CONV_MAX_NUM_DIMENSIONS]) \
+		SHADER_PARAMETER_ARRAY(FIntVector4, YDimension_YMemoryStride_XDimension_XMemoryStride, [NNXRT_CONV_MAX_NUM_DIMENSIONS]) \
+		SHADER_PARAMETER_ARRAY(FIntVector4, XBlockStartStride_XBlockStride_WDimension_WDimensionDilationXBlockStride, [NNXRT_CONV_MAX_NUM_DIMENSIONS]) \
+		SHADER_PARAMETER_ARRAY(FVector4f, OneDiv_GroupStride_GroupThreadStride_XBlockStride, [NNXRT_CONV_MAX_NUM_DIMENSIONS]) \
+		SHADER_PARAMETER(int32, NumWChannels) \
+		SHADER_PARAMETER(int32, YBatchStride) \
+		SHADER_PARAMETER(int32, YOutputKernelStride) \
+		SHADER_PARAMETER(int32, XBatchStride) \
+		SHADER_PARAMETER(int32, XChannelStride) \
+		SHADER_PARAMETER(int32, XBlockSize) \
+		SHADER_PARAMETER(int32, NumChannelBatches) \
+		SHADER_PARAMETER(int32, NumChannelsPerBatch) \
+		SHADER_PARAMETER(int32, WOutputKernelStride) \
+		SHADER_PARAMETER(int32, WChannelBatchSize) \
+		SHADER_PARAMETER(int32, WChannelSize) \
+		SHADER_PARAMETER(float, GroupsDivM) \
+	NNXRT_END_SHADER_PARAMETER_STRUCT()
+
+#ifdef NNXRT_RDG_DML
+NNXRT_CONV_PARAMETER_STRUCT()
+#endif
+
+//ConvTranspose
+#define NNXRT_CONVTRANSPOSE_MAX_NUM_STACK_DIMENSIONS 4
+#define NNXRT_CONVTRANSPOSE_MIN_NUM_READS_PER_THREAD_POW2 1
+#define NNXRT_CONVTRANSPOSE_MAX_NUM_READS_PER_THREAD_POW2 3
+
+#define NNXRT_CONVTRANSPOSE_PARAMETER_STRUCT() \
+	NNXRT_BEGIN_SHADER_PARAMETER_STRUCT(FMLConvTransposeParameters) \
+		NNXRT_RDG_BUFFER_SRV(Buffer<float>, X) \
+		NNXRT_RDG_BUFFER_SRV(Buffer<float>, W) \
+		NNXRT_RDG_BUFFER_UAV(RWBuffer<float>, Y) \
+		NNXRT_RDG_BUFFER_SRV(Buffer<float>, B) \
+		SHADER_PARAMETER_ARRAY(FIntVector4, Dilation_Stride_XBlockStartOffset_DilationXBlockStride, [NNXRT_CONVTRANSPOSE_MAX_NUM_STACK_DIMENSIONS]) \
+		SHADER_PARAMETER_ARRAY(FIntVector4, GroupStride_GroupShape_GroupThreadStride_StrideXBlockStride, [NNXRT_CONVTRANSPOSE_MAX_NUM_STACK_DIMENSIONS]) \
+		SHADER_PARAMETER_ARRAY(FIntVector4, YDimension_YMemoryStride_XDimension_XMemoryStride, [NNXRT_CONVTRANSPOSE_MAX_NUM_STACK_DIMENSIONS]) \
+		SHADER_PARAMETER_ARRAY(FIntVector4, XBlockStartStride_XBlockStride_WDimension_WDimensionDilationXBlockStride, [NNXRT_CONVTRANSPOSE_MAX_NUM_STACK_DIMENSIONS]) \
+		SHADER_PARAMETER_ARRAY(FVector4f, OneDiv_GroupStride_GroupThreadStride_OneDivStride, [NNXRT_CONVTRANSPOSE_MAX_NUM_STACK_DIMENSIONS]) \
+		SHADER_PARAMETER(int32, NumWChannels) \
+		SHADER_PARAMETER(int32, NumOutChannelsDivGroup) \
+		SHADER_PARAMETER(int32, YBatchStride) \
+		SHADER_PARAMETER(int32, YOutputKernelStride) \
+		SHADER_PARAMETER(int32, XBatchStride) \
+		SHADER_PARAMETER(int32, XChannelStride) \
+		SHADER_PARAMETER(int32, XBlockSize) \
+		SHADER_PARAMETER(int32, NumChannelBatches) \
+		SHADER_PARAMETER(int32, NumChannelsPerBatch) \
+		SHADER_PARAMETER(int32, WOutputKernelStride) \
+		SHADER_PARAMETER(int32, WChannelBatchSize) \
+		SHADER_PARAMETER(int32, WChannelSize) \
+		SHADER_PARAMETER(float, GroupsDivM) \
+		SHADER_PARAMETER(float, OneDivGroup) \
+	NNXRT_END_SHADER_PARAMETER_STRUCT()
+
+#ifdef NNXRT_RDG_DML
+NNXRT_CONVTRANSPOSE_PARAMETER_STRUCT()
+#endif
\ No newline at end of file
diff --git a/Engine/Plugins/Experimental/NNX/Source/NNIHlslShaders/Shaders/NNIHlslShadersBroadcastHelper.ush b/Engine/Plugins/Experimental/NNX/Source/NNIHlslShaders/Shaders/NNIHlslShadersBroadcastHelper.ush
new file mode 100644
--- /dev/null
+++ b/Engine/Plugins/Experimental/NNX/Source/NNIHlslShaders/Shaders/NNIHlslShadersBroadcastHelper.ush
@@ -0,0 +1,7 @@
+// Copyright Epic Games, Inc. All Rights Reserved.
+
+void DivMod(uint op1, uint op2, out uint q, out uint r)
+{
+	q = op1 / op2;
+	r = op1 - (q * op2);
+}
\ No newline at end of file
diff --git a/Engine/Plugins/Experimental/NNX/Source/NNIHlslShaders/Shaders/NNIHlslShadersConv.usf b/Engine/Plugins/Experimental/NNX/Source/NNIHlslShaders/Shaders/NNIHlslShadersConv.usf
new file mode 100644
--- /dev/null
+++ b/Engine/Plugins/Experimental/NNX/Source/NNIHlslShaders/Shaders/NNIHlslShadersConv.usf
@@ -0,0 +1,294 @@
+// Copyright Epic Games, Inc. All Rights Reserved.
+
+#include "/Engine/Public/Platform.ush"
+
+#define WORK_TYPE float
+#define BUFFER_TYPE float
+#define BUFFER_TO_WORK_TYPE(x) x
+#define WORK_TO_BUFFER_TYPE(x) x
+
+// Must correspond to EConvGroupSize defined in NNIHlslShadersConvCS.h
+#if GROUP_SIZE == 0
+#define NUM_GROUP_THREADS 128
+#elif GROUP_SIZE == 1
+#define NUM_GROUP_THREADS 256
+#elif GROUP_SIZE == 2
+#define NUM_GROUP_THREADS 512
+#endif
+
+Buffer<BUFFER_TYPE> X; // N x C x XD0 x ... x XDi
+Buffer<BUFFER_TYPE> W; // M x C/GROUPS x WD0 x ... x WDi
+RWBuffer<BUFFER_TYPE> Y;  // N x M x YD0 x ... x YDi
+Buffer<BUFFER_TYPE> B; // B0, ..., BM
+
+// x: Dilation0, ...., Dilationi
+// y: Strides0, ...., Stridesi
+// z: Pad0Begin, ...., PadiBegin
+int4 Dilation_Stride_XBlockStartOffset_DilationXBlockStride[MAX_NUM_DIMENSIONS];
+
+// x: Element j contains PROD(GDj+1,..,GDi) with the last element containing 1, with GDi being the number of thread groups contained by a dimension i
+// y: Number of threads in each dimension of a group
+// z: Element j contains PROD(GTDj+1,..,GTDi) with the last element containing 1, with GTDi being the number of threads contained by a dimension i inside a group
+int4 GroupStride_GroupShape_GroupThreadStride_StrideXBlockStride[MAX_NUM_DIMENSIONS];
+
+// x: YD0, ..., YDi
+// y: Element j contains PROD(YDj+1,..,YDi) with the last element containing 1
+// z: XD0, ..., XDi
+// w: Element j contains PROD(XDj+1,..,XDi) with the last element containing 1
+int4 YDimension_YMemoryStride_XDimension_XMemoryStride[MAX_NUM_DIMENSIONS];
+
+// x: Number of elements in each dimension of a X block
+// y: The strides of the X block to be loaded
+// z: WD0, ..., WDi
+int4 XBlockStartStride_XBlockStride_WDimension_WDimensionDilationXBlockStride[MAX_NUM_DIMENSIONS];
+
+// x: 1/GroupStride
+// y: 1/GroupThreadStride
+// z: 1/XBlockStride
+float4 OneDiv_GroupStride_GroupThreadStride_XBlockStride[MAX_NUM_DIMENSIONS];
+
+int NumWChannels; // WShape[1]
+
+int YBatchStride;
+int YOutputKernelStride;
+
+int XBatchStride; // The number of elements in each bacth of X
+int XChannelStride; // The number of elements in each channel of X
+
+int XBlockSize; // The total number of elements in each block to be loaded
+
+int NumChannelBatches; // The number of iteration needed to cover all channels that needs to be processed. Equals to ceil(NumWChannels/NumChannelsPerBatch)
+int NumChannelsPerBatch; // The number of complete kernel channels that can be loaded in one batch
+
+int WOutputKernelStride; // The total number of W elements per output kernel and thus PROD(C/GROUPS, WD0, .., WDi)
+int WChannelBatchSize; // The total number of W elements inside a channel batch and thus the number of channels per batch times PROD(WD0, .., WDi)
+int WChannelSize; // Number of elements in a single channel of W and thus PROD(WDi)
+
+float GroupsDivM; // The number of groups divided by the number of output kernels used to compute the start channel offset
+
+groupshared WORK_TYPE SharedMemoryX[NUM_GROUP_THREADS<<NUM_READS_PER_THREAD_POW2];
+groupshared WORK_TYPE SharedMemoryW[NUM_GROUP_THREADS];
+
+// Must correspond to EConvAlgorithm defined in NNIHlslShadersConvCS.h
+#if ALGORITHM < 1
+
+[numthreads(NUM_GROUP_THREADS, 1, 1)]
+void Conv(in const uint3 DispatchThreadID : SV_DispatchThreadID, in const uint3 GroupID : SV_GroupID, in const uint3 GroupThreadID : SV_GroupThreadID)
+{
+	// Group Z dimension and Y dimension are used for batch and output kernel
+    // This works as each group only works on the inner dimensions (num threads in y and z set to one, group count to N and M respectively)
+	int BatchIndex = GroupID.z;
+	int OutputKernelIndex = GroupID.y;
+
+	// Input channels and output kernels are divided into groups
+	// Thus the channel offset for this thread's output kernel can be computed as follows:
+	// The output kernel index is divided by the number of output kernels per group giving the group index
+	// This is multiplied by the number of channels in W to get the start index
+	int InputChannelStartIndex = NumWChannels * (int)(GroupsDivM * (float)OutputKernelIndex);
+
+	// Compute the output memory index and dimension indices of this thread inside its group as well as the start dimension indices of the X block to load
+	int YIndex = BatchIndex * YBatchStride + OutputKernelIndex * YOutputKernelStride;
+	int GroupThreadDimensionIndex[NUM_DIMENSIONS];// This will be used later to lookup the shared W and X to compute the final result
+    int XBlockStartDimensionIndex[NUM_DIMENSIONS];
+	int GroupIndex = GroupID.x;
+	int GroupThreadIndex = GroupThreadID.x;
+	UNROLL
+	for (int i = 0; i < NUM_DIMENSIONS; i++)
+	{
+        // Divide by the number of groups contained by one index of this dimension to get the dimension index n being the nth group in this dimension
+		int GroupDimensionIndex = (int)(OneDiv_GroupStride_GroupThreadStride_XBlockStride[i].x * (float)GroupIndex);
+		GroupIndex -= GroupDimensionIndex * GroupStride_GroupShape_GroupThreadStride_StrideXBlockStride[i].x;
+
+        // Same here with number of threads contained inside the group giving the nth thread inside the group in this dimension
+		GroupThreadDimensionIndex[i] = (int)(OneDiv_GroupStride_GroupThreadStride_XBlockStride[i].y * (float)GroupThreadIndex);
+		GroupThreadIndex -= GroupThreadDimensionIndex[i] * GroupStride_GroupShape_GroupThreadStride_StrideXBlockStride[i].z;
+		
+        // Compose the global index by the group index times group shape plus thread index inside the group
+		int ThreadDimensionIndex = GroupDimensionIndex * GroupStride_GroupShape_GroupThreadStride_StrideXBlockStride[i].y + GroupThreadDimensionIndex[i];
+        
+        // Subtract the padding from the index of the first element of the group to get the start index of the block needed by this group
+        XBlockStartDimensionIndex[i] = GroupDimensionIndex * XBlockStartStride_XBlockStride_WDimension_WDimensionDilationXBlockStride[i].x + Dilation_Stride_XBlockStartOffset_DilationXBlockStride[i].z;
+
+        // ThreadDimensionIndex is always positive but needs to be tested if it does not overflow
+		if (ThreadDimensionIndex >= YDimension_YMemoryStride_XDimension_XMemoryStride[i].x || YIndex < 0)
+		{
+			YIndex = -1;
+		}
+		else
+		{
+			YIndex += ThreadDimensionIndex * YDimension_YMemoryStride_XDimension_XMemoryStride[i].y;
+		}
+	}
+
+    // Compute the memory indices in global X which the current thread has to load as well as the indices in shared X where to write
+	int GlobalXIndex[0x1<<NUM_READS_PER_THREAD_POW2];
+	int SharedXIndex[0x1<<NUM_READS_PER_THREAD_POW2];
+	int ReadIdx = GroupThreadID.x;
+    UNROLL
+    for (int j = 0; j < (int)(0x1<<NUM_READS_PER_THREAD_POW2); j++)
+	{
+        // If the read index exceeded the number of elements in the X block further computation can be skipped
+        if (ReadIdx >= XBlockSize)
+        {
+            // Assign -1 for both to indicate that the element is not used and thus neither has to be read from global nor written to shared memory
+            GlobalXIndex[j] = -1;
+            SharedXIndex[j] = -1;
+            continue;
+        }
+
+        // Turn the flat group index multiple into thread dimension indices inside the volume to be loaded by this group
+        int TmpGlobalXIndex = BatchIndex * XBatchStride + InputChannelStartIndex * XChannelStride;
+        int TmpSharedXIndex = 0;
+        int TempReadIdx = ReadIdx;
+        UNROLL
+        for (int k = 0; k < NUM_DIMENSIONS; k++)
+        {
+            int DimensionIndex = (int)(OneDiv_GroupStride_GroupThreadStride_XBlockStride[k].z * (float)TempReadIdx);
+            int DimensionIndexOffset = DimensionIndex * XBlockStartStride_XBlockStride_WDimension_WDimensionDilationXBlockStride[k].y;
+		    TempReadIdx -= DimensionIndexOffset;
+
+            // Update TmpGlobalXIndex only if positive (it is already outside the volume if negative)
+            if (TmpGlobalXIndex >= 0)
+            {
+                // Check if the global X index is in range
+                int XIndex = XBlockStartDimensionIndex[k] + DimensionIndex;
+                if (XIndex < 0 || XIndex >= YDimension_YMemoryStride_XDimension_XMemoryStride[k].z)
+                {
+                    TmpGlobalXIndex = -1;
+                }
+                else
+                {
+                    TmpGlobalXIndex += XIndex * YDimension_YMemoryStride_XDimension_XMemoryStride[k].w;
+                }
+            }
+
+            // Update TmpSharedXIndex in any case as a zero has to be written
+            TmpSharedXIndex += DimensionIndexOffset;
+        }
+
+        // Assign the result, which will be -1 for the global and positive for the shared index if the element lies in the padding area
+        // This indicates that the global element should not be loaded but a zero has to be written to shared memory 
+        GlobalXIndex[j] = TmpGlobalXIndex;
+        SharedXIndex[j] = TmpSharedXIndex;
+
+        // Increment the read idx by the number of threads inside the group to cover all elements
+        ReadIdx += NUM_GROUP_THREADS;
+    }
+
+    // Write the zeros once to avoid rewriting the same value later
+    UNROLL
+    for (int l = 0; l < (int)(0x1<<NUM_READS_PER_THREAD_POW2); l++)
+	{
+        // If the global index is out of bound (padding area) but the shared index is valid (inside the X block) the shared element needs initialization to zero
+        if (GlobalXIndex[l] < 0 && SharedXIndex[l] >= 0)
+        {
+            SharedMemoryX[SharedXIndex[l]] = 0;
+        }
+    }
+
+    // Iterate over all channel batches
+    int WChannelOutputKernelOffset = OutputKernelIndex * WOutputKernelStride;
+    int WChannelBatchOffset = 0;
+    int WChannelIndex = 0;
+#if HAS_B == 0
+    WORK_TYPE Result = 0;
+#else
+    WORK_TYPE Result = BUFFER_TO_WORK_TYPE(B[OutputKernelIndex]);
+#endif
+    for (int m = 0; m < NumChannelBatches; m++)
+	{
+        // Load all channels of W belonging to this batch
+        int WIndex = WChannelBatchOffset + GroupThreadID.x;
+        if (GroupThreadID.x < WChannelBatchSize && WIndex < WOutputKernelStride)
+        {
+            SharedMemoryW[GroupThreadID.x] = BUFFER_TO_WORK_TYPE(W[WChannelOutputKernelOffset + WIndex]);
+            WChannelBatchOffset += WChannelBatchSize;
+        }
+        
+        // Each channel batch iterates over all consecutive elements of the loaded kernel channels
+        // This results in bank conflict free lookups, as all threads of the same warp read exactly the same element which is then broadcasted
+        int SharedWIndex = 0;
+
+        // Iterate over all channels in this batch
+        for (int n = 0; n < NumChannelsPerBatch; n++)
+	    {
+            // Check if the channel index is valid as channel batches may not be aligned with the effective number of channels
+            if (WChannelIndex >= NumWChannels)
+            {
+                break;
+            }
+            WChannelIndex++;
+            
+            // Load the shared X
+            UNROLL
+            for (int o = 0; o < (int)(0x1<<NUM_READS_PER_THREAD_POW2); o++)
+	        {
+                // If the global index is negative, a zero has already been written, if shared index is negative there is no need to load anything
+                if (GlobalXIndex[o] >= 0 && SharedXIndex[o] >= 0)
+                {
+                    // Load the global data into shared memory
+                    SharedMemoryX[SharedXIndex[o]] = BUFFER_TO_WORK_TYPE(X[GlobalXIndex[o]]);
+
+                    // Update the indices for the next iteration by adding the channel stride
+                    GlobalXIndex[o] += XChannelStride;
+                }
+            }
+
+            // Sync the group before doing any computation
+            GroupMemoryBarrierWithGroupSync();
+
+            // Initialize a dimension index array used to iterate through the kernel dimensions
+            // Also initialize the start index in shared memory based on the result position
+            int WDimensionIndex[NUM_DIMENSIONS];
+            int SharedXIndex = 0;
+            UNROLL
+            for (int p = 0; p < NUM_DIMENSIONS; p++)
+            {
+                WDimensionIndex[p] = 0;
+
+                // The index of the first element (e.g. upper left) required by each output is offsetted by the kernel stride and then multiplied by the stride in memory
+                SharedXIndex += GroupThreadDimensionIndex[p] * GroupStride_GroupShape_GroupThreadStride_StrideXBlockStride[p].w;
+            }
+
+            // Compute and accumulate the result
+            for (int q = 0; q < WChannelSize; q++)
+            {
+                Result += SharedMemoryX[SharedXIndex] * SharedMemoryW[SharedWIndex];
+
+                // Advance the inner-most dimension
+                SharedWIndex++;
+                SharedXIndex += Dilation_Stride_XBlockStartOffset_DilationXBlockStride[NUM_DIMENSIONS-1].x;
+                WDimensionIndex[NUM_DIMENSIONS-1]++;
+
+                // Advance the other dimensions if necessary
+                UNROLL
+                for (int r = NUM_DIMENSIONS-1; r > 0; r--)
+                {
+                    if(WDimensionIndex[r] >= XBlockStartStride_XBlockStride_WDimension_WDimensionDilationXBlockStride[r].z)
+                    {
+                        // Adjust the SharedXIndex by resetting the current dimension and increase the next outer one
+                        SharedXIndex -= XBlockStartStride_XBlockStride_WDimension_WDimensionDilationXBlockStride[r].w;
+                        SharedXIndex += Dilation_Stride_XBlockStartOffset_DilationXBlockStride[r-1].w;
+                        WDimensionIndex[r] = 0;
+                        WDimensionIndex[r-1]++;
+                    }
+                    else 
+                    {
+                        break;
+                    }
+                }
+            }
+
+            // Sync the group before continuing any load
+            GroupMemoryBarrierWithGroupSync();
+        }
+    }
+
+    // Write the final result
+    if (YIndex >= 0)
+    {
+        Y[YIndex] = WORK_TO_BUFFER_TYPE(Result);
+    }
+}
+
+#endif
\ No newline at end of file
diff --git a/Engine/Plugins/Experimental/NNX/Source/NNIHlslShaders/Shaders/NNIHlslShadersConvTranspose.usf b/Engine/Plugins/Experimental/NNX/Source/NNIHlslShaders/Shaders/NNIHlslShadersConvTranspose.usf
new file mode 100644
--- /dev/null
+++ b/Engine/Plugins/Experimental/NNX/Source/NNIHlslShaders/Shaders/NNIHlslShadersConvTranspose.usf
@@ -0,0 +1,334 @@
+// Copyright Epic Games, Inc. All Rights Reserved.
+
+#include "/Engine/Public/Platform.ush"
+
+#if DATA_TYPE == 0
+
+#define WORK_TYPE float
+#define BUFFER_TYPE float
+#define BUFFER_TO_WORK_TYPE(x) x
+#define WORK_TO_BUFFER_TYPE(x) x
+
+#elif DATA_TYPE == 1
+
+#define WORK_TYPE int
+#define BUFFER_TYPE int
+#define BUFFER_TO_WORK_TYPE(x) x
+#define WORK_TO_BUFFER_TYPE(x) x
+
+#elif DATA_TYPE == 2
+
+#define WORK_TYPE uint
+#define BUFFER_TYPE uint
+#define BUFFER_TO_WORK_TYPE(x) x
+#define WORK_TO_BUFFER_TYPE(x) x
+
+#elif DATA_TYPE == 3
+
+#define WORK_TYPE float
+#define BUFFER_TYPE min16uint
+#define BUFFER_TO_WORK_TYPE(x) f16tof32(x)
+#define WORK_TO_BUFFER_TYPE(x) f32tof16(x)
+
+#endif
+
+// Must correspond to EConvGroupSize defined in NNIHlslShadersConvCS.h
+#if GROUP_SIZE == 0
+#define NUM_GROUP_THREADS 128
+#elif GROUP_SIZE == 1
+#define NUM_GROUP_THREADS 256
+#elif GROUP_SIZE == 2
+#define NUM_GROUP_THREADS 512
+#endif
+
+Buffer<BUFFER_TYPE> X; // N x C x XD0 x ... x XDi
+Buffer<BUFFER_TYPE> W; // M x C/GROUPS x WD0 x ... x WDi
+RWBuffer<BUFFER_TYPE> Y;  // N x M x YD0 x ... x YDi
+Buffer<BUFFER_TYPE> B; // B0, ..., BM
+
+// x: Dilation0, ...., Dilationi
+// y: Strides0, ...., Stridesi
+// z: Pad0Begin, ...., PadiBegin
+int4 Dilation_Stride_XBlockStartOffset_DilationXBlockStride[NUM_STACK_DIMENSIONS];
+
+// x: Element j contains PROD(GDj+1,..,GDi) with the last element containing 1, with GDi being the number of thread groups contained by a dimension i
+// y: Number of threads in each dimension of a group
+// z: Element j contains PROD(GTDj+1,..,GTDi) with the last element containing 1, with GTDi being the number of threads contained by a dimension i inside a group
+int4 GroupStride_GroupShape_GroupThreadStride_StrideXBlockStride[NUM_STACK_DIMENSIONS];
+
+// x: YD0, ..., YDi
+// y: Element j contains PROD(YDj+1,..,YDi) with the last element containing 1
+// z: XD0, ..., XDi
+// w: Element j contains PROD(XDj+1,..,XDi) with the last element containing 1
+int4 YDimension_YMemoryStride_XDimension_XMemoryStride[NUM_STACK_DIMENSIONS];
+
+// x: Number of elements in each dimension of a X block
+// y: The strides of the X block to be loaded
+// z: WD0, ..., WDi
+int4 XBlockStartStride_XBlockStride_WDimension_WDimensionDilationXBlockStride[NUM_STACK_DIMENSIONS];
+
+// x: 1/GroupStride
+// y: 1/GroupThreadStride
+// z: 1/XBlockStride
+float4 OneDiv_GroupStride_GroupThreadStride_OneDivStride[NUM_STACK_DIMENSIONS];
+
+int NumWChannels; // WShape[0]
+int NumOutChannelsDivGroup; // WShape[1]
+
+int YBatchStride;
+int YOutputKernelStride;
+
+int XBatchStride; // The number of elements in each bacth of X
+int XChannelStride; // The number of elements in each channel of X
+
+int XBlockSize; // The total number of elements in each block to be loaded
+
+int NumChannelBatches; // The number of iteration needed to cover all channels that needs to be processed. Equals to ceil(NumWChannels/NumChannelsPerBatch)
+int NumChannelsPerBatch; // The number of complete kernel channels that can be loaded in one batch
+
+int WOutputKernelStride; // The total number of W elements per output kernel and thus PROD(C/GROUPS, WD0, .., WDi)
+int WChannelBatchSize; // The total number of W elements inside a channel batch and thus the number of channels per batch times PROD(WD0, .., WDi)
+int WChannelSize; // Number of elements in a single channel of W and thus PROD(WDi)
+
+float GroupsDivM; // The number of groups divided by the number of output kernels used to compute the start channel offset
+float OneDivGroup;
+
+groupshared WORK_TYPE SharedMemoryX[NUM_GROUP_THREADS<<NUM_READS_PER_THREAD_POW2];
+groupshared WORK_TYPE SharedMemoryW[NUM_GROUP_THREADS];
+
+// Must correspond to EConvAlgorithm defined in NNIHlslShadersConvCS.h
+#if ALGORITHM < 1
+
+[numthreads(NUM_GROUP_THREADS, 1, 1)]
+void ConvTranspose(in const int3 DispatchThreadID : SV_DispatchThreadID, in const int3 GroupID : SV_GroupID, in const int3 GroupThreadID : SV_GroupThreadID)
+{
+	// Group Z dimension and Y dimension are used for batch and output kernel
+    // This works as each group only works on the inner dimensions (num threads in y and z set to one, group count to N and M respectively)
+	int BatchIndex = GroupID.z;
+	int OutputKernelIndex = GroupID.y;
+
+	// Input channels and output kernels are divided into groups
+	// Thus the channel offset for this thread's output kernel can be computed as follows:
+	// The output kernel index is divided by the number of output kernels per group giving the group index
+	// This is multiplied by the number of channels in W to get the start index
+	int InputChannelStartIndex = (int)(GroupsDivM * (float)OutputKernelIndex);
+
+	// Compute the output memory index and dimension indices of this thread inside its group as well as the start dimension indices of the X block to load
+	int YIndex = BatchIndex * YBatchStride + OutputKernelIndex * YOutputKernelStride;
+	int GroupThreadDimensionIndex[NUM_STACK_DIMENSIONS];// This will be used later to lookup the shared W and X to compute the final result
+    int XBlockStartDimensionIndex[NUM_STACK_DIMENSIONS];
+    float ResidualXIndex[NUM_STACK_DIMENSIONS];
+	int GroupIndex = GroupID.x;
+	int GroupThreadIndex = GroupThreadID.x;
+	UNROLL
+	for (int i = 0; i < NUM_STACK_DIMENSIONS; i++)
+	{
+        // Divide by the number of groups contained by one index of this dimension to get the dimension index n being the nth group in this dimension
+		int GroupDimensionIndex = (int)(OneDiv_GroupStride_GroupThreadStride_OneDivStride[i].x *(float)GroupIndex);
+		GroupIndex -= GroupDimensionIndex * GroupStride_GroupShape_GroupThreadStride_StrideXBlockStride[i].x;
+
+        // Same here with number of threads contained inside the group giving the nth thread inside the group in this dimension
+		GroupThreadDimensionIndex[i] = OneDiv_GroupStride_GroupThreadStride_OneDivStride[i].y * (float)GroupThreadIndex;
+		GroupThreadIndex -= GroupThreadDimensionIndex[i] * (float)GroupStride_GroupShape_GroupThreadStride_StrideXBlockStride[i].z;
+
+        float FXBlockStartDimensionIndex = ((float)GroupDimensionIndex * (float)XBlockStartStride_XBlockStride_WDimension_WDimensionDilationXBlockStride[i].x + (float)Dilation_Stride_XBlockStartOffset_DilationXBlockStride[i].z) * OneDiv_GroupStride_GroupThreadStride_OneDivStride[i].w;
+        ResidualXIndex[i] = (float)(FXBlockStartDimensionIndex - int(FXBlockStartDimensionIndex)) * (float)Dilation_Stride_XBlockStartOffset_DilationXBlockStride[i].y;
+
+        // Compose the global index by the group index times group shape plus thread index inside the group
+		int ThreadDimensionIndex = GroupDimensionIndex * GroupStride_GroupShape_GroupThreadStride_StrideXBlockStride[i].y + (int)GroupThreadDimensionIndex[i];
+
+        // Subtract the padding from the index of the first element of the group to get the start index of the block needed by this group
+        XBlockStartDimensionIndex[i] = (int)FXBlockStartDimensionIndex;
+
+        // ThreadDimensionIndex is always positive but needs to be tested if it does not overflow
+		if (ThreadDimensionIndex >= YDimension_YMemoryStride_XDimension_XMemoryStride[i].x || YIndex < 0)
+		{
+			YIndex = -1;
+		}
+		else
+		{
+			YIndex += ThreadDimensionIndex * YDimension_YMemoryStride_XDimension_XMemoryStride[i].y;
+		}
+	}
+
+    // Compute the memory indices in global X which the current thread has to load as well as the indices in shared X where to write
+	int GlobalXIndex[0x1<<NUM_READS_PER_THREAD_POW2];
+	int SharedXIndex[0x1<<NUM_READS_PER_THREAD_POW2];
+	int ReadIdx = GroupThreadID.x;
+    UNROLL
+    for (int j = 0; j < (int)(0x1<<NUM_READS_PER_THREAD_POW2); j++)
+	{
+        // If the read index exceeded the number of elements in the X block further computation can be skipped
+        if (ReadIdx >= XBlockSize)
+        {
+            // Assign -1 for both to indicate that the element is not used and thus neither has to be read from global nor written to shared memory
+            GlobalXIndex[j] = -1;
+            SharedXIndex[j] = -1;
+            continue;
+        }
+
+        // Turn the flat group index multiple into thread dimension indices inside the volume to be loaded by this group
+        int TmpGlobalXIndex = BatchIndex * XBatchStride + InputChannelStartIndex * XChannelStride;
+        int TmpSharedXIndex = 0;
+        int TempReadIdx = ReadIdx;
+        UNROLL
+        for (int k = 0; k < NUM_STACK_DIMENSIONS; k++)
+        {
+            int DimensionIndex = (int)(OneDiv_GroupStride_GroupThreadStride_OneDivStride[k].z * (float)TempReadIdx);
+            int DimensionIndexOffset = DimensionIndex * XBlockStartStride_XBlockStride_WDimension_WDimensionDilationXBlockStride[k].y;
+		    TempReadIdx -= DimensionIndexOffset;
+
+            // Update TmpGlobalXIndex only if positive (it is already outside the volume if negative)
+            if (TmpGlobalXIndex >= 0)
+            {
+                // Check if the global X index is in range
+                int XIndex = XBlockStartDimensionIndex[k] + DimensionIndex;
+                if (XIndex < 0 || XIndex >= YDimension_YMemoryStride_XDimension_XMemoryStride[k].z)
+                {
+                    TmpGlobalXIndex = -1;
+                }
+                else
+                {
+                    TmpGlobalXIndex += XIndex * YDimension_YMemoryStride_XDimension_XMemoryStride[k].w;
+                }
+            }
+
+            // Update TmpSharedXIndex in any case as a zero has to be written
+            TmpSharedXIndex += DimensionIndexOffset;
+        }
+
+        // Assign the result, which will be -1 for the global and positive for the shared index if the element lies in the padding area
+        // This indicates that the global element should not be loaded but a zero has to be written to shared memory
+        GlobalXIndex[j] = TmpGlobalXIndex;
+        SharedXIndex[j] = TmpSharedXIndex;
+
+        // Increment the read idx by the number of threads inside the group to cover all elements
+        ReadIdx += NUM_GROUP_THREADS;
+    }
+ 
+    // Write the zeros once to avoid rewriting the same value later
+    UNROLL
+    for (int l = 0; l < (int)(0x1<<NUM_READS_PER_THREAD_POW2); l++)
+	{
+        // If the global index is out of bound (padding area) but the shared index is valid (inside the X block) the shared element needs initialization to zero
+        if (GlobalXIndex[l] < 0 && SharedXIndex[l] >= 0)
+        {
+            SharedMemoryX[SharedXIndex[l]] = 0;
+        }
+    }
+
+    // Iterate over all channel batches
+    int WChannelBatchOffset = 0;
+    int WChannelIndex = 0;
+#if HAS_B == 0
+    WORK_TYPE Result = 0;
+#else
+    WORK_TYPE Result = BUFFER_TO_WORK_TYPE(B[OutputKernelIndex]);
+#endif
+    // Each channel batch iterates over all consecutive elements of the loaded kernel channels
+    // This results in bank conflict free lookups, as all threads of the same warp read exactly the same element which is then broadcasted
+    for (int m = 0; m < NumChannelBatches; m++)
+	{
+        // Load all channels of W belonging to this batch
+        int SharedWIndex = 0;
+        WChannelBatchOffset = m * WChannelBatchSize;
+
+        // Iterate over all channels in this batch
+        for (int n = 0; n < NumChannelsPerBatch; n++)
+	    {
+            int WChannelOutputKernelOffset = WChannelSize * OutputKernelIndex + NumOutChannelsDivGroup * int((float)n * OneDivGroup);
+            int WIndex = WChannelBatchOffset + GroupThreadID.x;
+            if (GroupThreadID.x < WChannelBatchSize && WIndex < WOutputKernelStride)
+            {  
+                SharedMemoryW[GroupThreadID.x] = BUFFER_TO_WORK_TYPE(W[WChannelOutputKernelOffset + WIndex]);
+
+            }
+            if (WChannelOutputKernelOffset < WChannelSize * NumOutChannelsDivGroup * n ||
+                WChannelOutputKernelOffset >= NumOutChannelsDivGroup * WChannelSize * (n + 1)) {
+                SharedMemoryW[GroupThreadID.x] = 0;
+            }
+
+            // Check if the channel index is valid as channel batches may not be aligned with the effective number of channels
+            if (WChannelIndex >= NumWChannels)
+            {
+                break;
+            }
+            WChannelIndex++;
+
+            // Load the shared X
+            UNROLL
+            for (int o = 0; o < (int)(0x1<<NUM_READS_PER_THREAD_POW2); o++)
+	        {
+                // If the global index is negative, a zero has already been written, if shared index is negative there is no need to load anything
+                if (GlobalXIndex[o] >= 0 && SharedXIndex[o] >= 0)
+                {
+                    // Load the global data into shared memory
+                    SharedMemoryX[SharedXIndex[o]] = BUFFER_TO_WORK_TYPE(X[GlobalXIndex[o]]);
+
+                    // Update the indices for the next iteration by adding the channel stride
+                    GlobalXIndex[o] += XChannelStride;
+                }
+            }
+
+            // Sync the group before doing any computation
+            GroupMemoryBarrierWithGroupSync();
+
+            // Initialize a dimension index array used to iterate through the kernel dimensions
+            // Also initialize the start index in shared memory based on the result position
+            int WDimensionIndex[NUM_STACK_DIMENSIONS];
+            float SharedXIndex = 0;
+
+            UNROLL
+            for (int p = 0; p < NUM_STACK_DIMENSIONS; p++)
+            {
+                WDimensionIndex[p] = 0;
+
+                // The index of the first element (e.g. upper left) required by each output is offsetted by the kernel stride and then multiplied by the stride in memory
+                SharedXIndex += (GroupThreadDimensionIndex[p] + ResidualXIndex[p]) * OneDiv_GroupStride_GroupThreadStride_OneDivStride[p].w * (float)GroupStride_GroupShape_GroupThreadStride_StrideXBlockStride[p].w;
+                }
+
+            // Compute and accumulate the result
+            for (int q = 0; q < WChannelSize; q++)
+            {
+                if (SharedXIndex - int(SharedXIndex+0.0001) < 0.001 && SharedXIndex > -0.0001){
+                    Result += SharedMemoryX[SharedXIndex+0.0001] * SharedMemoryW[WChannelSize - 1 - q];
+                }
+
+                // Advance the inner-most dimension
+                SharedWIndex++;
+                SharedXIndex += Dilation_Stride_XBlockStartOffset_DilationXBlockStride[NUM_STACK_DIMENSIONS-1].x * OneDiv_GroupStride_GroupThreadStride_OneDivStride[NUM_STACK_DIMENSIONS-1].w;
+
+                WDimensionIndex[NUM_STACK_DIMENSIONS-1]++;
+
+                // Advance the other dimensions if necessary
+                UNROLL
+                for (int r = NUM_STACK_DIMENSIONS-1; r > 0; r--)
+                {
+                    if(WDimensionIndex[r] >= XBlockStartStride_XBlockStride_WDimension_WDimensionDilationXBlockStride[r].z)
+                    {
+                        // Adjust the SharedXIndex by resetting the current dimension and increase the next outer one
+                        SharedXIndex -= XBlockStartStride_XBlockStride_WDimension_WDimensionDilationXBlockStride[r].w * OneDiv_GroupStride_GroupThreadStride_OneDivStride[r].w;
+                        SharedXIndex += Dilation_Stride_XBlockStartOffset_DilationXBlockStride[r-1].w * OneDiv_GroupStride_GroupThreadStride_OneDivStride[r-1].w;
+                        WDimensionIndex[r] = 0;
+                        WDimensionIndex[r-1]++;
+                    }
+                    else
+                    {
+                        break;
+                    }
+                }
+            }
+
+            // Sync the group before continuing any load
+            GroupMemoryBarrierWithGroupSync();
+        }
+    }
+
+    // Write the final result
+    if (YIndex >= 0)
+    {
+        Y[YIndex] = WORK_TO_BUFFER_TYPE(Result);
+    }
+}
+
+#endif
diff --git a/Engine/Plugins/Experimental/NNX/Source/NNIHlslShaders/Shaders/NNIHlslShadersElementWise.usf b/Engine/Plugins/Experimental/NNX/Source/NNIHlslShaders/Shaders/NNIHlslShadersElementWise.usf
new file mode 100644
--- /dev/null
+++ b/Engine/Plugins/Experimental/NNX/Source/NNIHlslShaders/Shaders/NNIHlslShadersElementWise.usf
@@ -0,0 +1,118 @@
+// Copyright Epic Games, Inc. All Rights Reserved.
+
+#include "/Engine/Public/Platform.ush"
+
+Buffer<float> Input;
+RWBuffer<float> Output;
+uint Num;
+uint ThreadCountX;
+float Alpha;
+float Beta;
+float Gamma;
+
+//see [REDACTED_PROFILE_URL]/onnx/blob/main/docs/Operators.md#Relu
+float relu(float x)
+{
+	return max(0.0f, x);
+}
+
+//see [REDACTED_PROFILE_URL]/onnx/blob/main/docs/Operators.md#Selu
+float selu(float x)
+{
+	float yNegOrZero = Gamma * (Alpha * exp(x) - Alpha);
+	float yPos = Gamma * x;
+	return (x > 0.0f) ? yPos : yNegOrZero;
+}
+
+//see [REDACTED_PROFILE_URL]/onnx/blob/main/docs/Operators.md#sigmoid
+float sigmoid(float x)
+{
+	return 1.0f / (1.0f + exp(-x));
+}
+
+//see [REDACTED_PROFILE_URL]/onnx/blob/main/docs/Operators.md#softplus
+float softplus(float x)
+{
+	return log(exp(x) + 1.0f);
+}
+
+//see [REDACTED_PROFILE_URL]/onnx/blob/main/docs/Operators.md#softsign
+float softsign(float x)
+{
+	return x/(1.0f+abs(x));
+}
+
+//see [REDACTED_PROFILE_URL]/onnx/blob/main/docs/Operators.md#acosh
+float acosh(float x)
+{
+	//https://mathworld.wolfram.com/InverseHyperbolicCosine.html
+	const static float NaN = 0.0f / 0.0f;
+	float yAboveOne = log( x + sqrt( x + 1.0f ) + sqrt( x - 1.0f ) );
+	return (x == 1.0f) ? 0.0f : (x >= 1.0f) ? yAboveOne : NaN;
+}
+
+//see [REDACTED_PROFILE_URL]/onnx/blob/main/docs/Operators.md#asinh
+float asinh(float x)
+{
+	//https://mathworld.wolfram.com/InverseHyperbolicSine.html
+	return log( x + sqrt( 1 + ( x * x ) ) );
+}
+
+//see [REDACTED_PROFILE_URL]/onnx/blob/main/docs/Operators.md#atanh
+float atanh(float x)
+{
+	//https://mathworld.wolfram.com/InverseHyperbolicTangent.html
+	return 0.5f * ( log( 1 + x) - log ( 1 - x) );
+}
+
+//see [REDACTED_PROFILE_URL]/onnx/blob/main/docs/Operators.md#elu
+float elu(float x)
+{
+    float yNeg = Alpha * (exp(x) - 1.0f);
+	float yPosOrZero = x;
+	return (x >= 0.0f) ? yPosOrZero : yNeg;
+}
+
+//see [REDACTED_PROFILE_URL]/onnx/blob/main/docs/Operators.md#erf
+float erf(float x)
+{
+	//https://aapt.scitation.org/doi/abs/10.1119/1.15018?journalCode=ajp
+	float a = 167.0f/148.0f;
+	float b = 11.0f/109.0f;
+	float x3 = x*x*x;
+	return tanh( a*x + b*x3 );
+}
+
+float hardSigmoid(float x, float alpha, float beta)
+{
+	return max(0.0f, min(1.0f, alpha * x + beta));
+}
+
+//see [REDACTED_PROFILE_URL]/onnx/blob/main/docs/Operators.md#hardSigmoid
+float hardSigmoid(float x)
+{
+	return hardSigmoid(x, Alpha, Beta);
+}
+
+//see [REDACTED_PROFILE_URL]/onnx/blob/main/docs/Operators.md#hardSwish
+float hardSwish(float x)
+{
+    return x * hardSigmoid(x, 1.0f/6.0f, 0.5f);
+}
+
+//see [REDACTED_PROFILE_URL]/onnx/blob/main/docs/Operators.md#leakyRelu
+float leakyRelu(float x)
+{
+    return (x >= 0.0f) ? x : Alpha * x;
+}
+
+
+[numthreads(THREADGROUP_SIZE_X, 1, 1)]
+void ElementWise(in const uint3 DispatchThreadID : SV_DispatchThreadID)
+{
+	const uint Index = DispatchThreadID.y * ThreadCountX + DispatchThreadID.x;
+	if (Index < Num)
+	{
+		Output[Index] = ELEMENTWISE_OP(Input[Index]);
+	}
+}
diff --git a/Engine/Plugins/Experimental/NNX/Source/NNIHlslShaders/Shaders/NNIHlslShadersElementWiseBinary.usf b/Engine/Plugins/Experimental/NNX/Source/NNIHlslShaders/Shaders/NNIHlslShadersElementWiseBinary.usf
new file mode 100644
--- /dev/null
+++ b/Engine/Plugins/Experimental/NNX/Source/NNIHlslShaders/Shaders/NNIHlslShadersElementWiseBinary.usf
@@ -0,0 +1,61 @@
+// Copyright Epic Games, Inc. All Rights Reserved.
+
+#include "/Engine/Public/Platform.ush"
+#include "NNIHlslShadersBroadcastHelper.ush"
+
+Buffer<float> LHSInput;
+Buffer<float> RHSInput;
+RWBuffer<float> Output;
+uint4 TensorInfo[NUM_DIMENSIONS];
+uint Num;
+uint ThreadCountX;
+
+// See [REDACTED_PROFILE_URL]/onnx/blob/main/docs/Operators.md#prelu
+float prelu(float x, float slope)
+{
+	return (x < 0.0f) ? (slope * x) : x;
+}
+
+float safe_pow(float base, float xp)
+{
+	if (base >= 0.0f)
+		return pow(base, xp);
+	
+	//negative base + not integer exponent => nan
+	if (round(xp) != xp)
+	{
+		
+		const float nan = 0. / 0.;
+		return nan;
+	}
+
+	//negative base + integer exponent
+	int isExpOdd = (xp % 2);
+	float sign = isExpOdd ? -1.0f : 1.0f;
+	return sign * pow(-base, xp);
+}
+
+[numthreads(THREADGROUP_SIZE_X, 1, 1)]
+void ElementWiseBinary(in const uint3 DispatchThreadID : SV_DispatchThreadID)
+{
+	const uint Index = DispatchThreadID.y * ThreadCountX + DispatchThreadID.x;
+	
+	if (Index < Num)
+	{
+		uint LHSIndex = 0;
+		uint RHSIndex = 0;
+		uint Offset = Index;
+		for (uint dim = 0; dim < NUM_DIMENSIONS; ++dim)
+		{
+			uint Q, R;
+			DivMod(Offset, TensorInfo[dim][2], Q, R);
+			LHSIndex += TensorInfo[dim][0] * Q;
+			RHSIndex += TensorInfo[dim][1] * Q;
+			Offset = R;
+		}
+
+		float X = LHSInput[LHSIndex];
+		float Y = RHSInput[RHSIndex];
+		Output[Index] = ELEMENTWISE_OP(X, Y);
+	}
+}
diff --git a/Engine/Plugins/Experimental/NNX/Source/NNIHlslShaders/Shaders/NNIHlslShadersElementWiseVariadic.usf b/Engine/Plugins/Experimental/NNX/Source/NNIHlslShaders/Shaders/NNIHlslShadersElementWiseVariadic.usf
new file mode 100644
--- /dev/null
+++ b/Engine/Plugins/Experimental/NNX/Source/NNIHlslShaders/Shaders/NNIHlslShadersElementWiseVariadic.usf
@@ -0,0 +1,65 @@
+// Copyright Epic Games, Inc. All Rights Reserved.
+
+#include "/Engine/Public/Platform.ush"
+#include "NNIHlslShadersBroadcastHelper.ush"
+
+#define BUFFER_TYPE float
+
+Buffer<BUFFER_TYPE> Input0;
+Buffer<BUFFER_TYPE> Input1;
+Buffer<BUFFER_TYPE> Input2;
+Buffer<BUFFER_TYPE> Input3;
+RWBuffer<BUFFER_TYPE> Output;
+uint4 InputTensorInfo[NUM_DIMENSIONS];
+uint4 OutputTensorInfo[NUM_DIMENSIONS];
+uint Num;
+uint ThreadCountX;
+float Scale;
+
+[numthreads(THREADGROUP_SIZE_X, 1, 1)]
+void ElementWiseVariadic(in const uint3 DispatchThreadID : SV_DispatchThreadID)
+{
+	const uint Index = DispatchThreadID.y * ThreadCountX + DispatchThreadID.x;
+	
+	if (Index < Num)
+	{
+		uint Input0Index = 0;
+		uint Input1Index = 0;
+		uint Input2Index = 0;
+		uint Input3Index = 0;
+		uint Offset = Index;
+		for (uint dim = 0; dim < NUM_DIMENSIONS; ++dim)
+		{
+			uint Q, R;
+			DivMod(Offset, OutputTensorInfo[dim][0], Q, R);
+			Input0Index += InputTensorInfo[dim][0] * Q;
+			Input1Index += InputTensorInfo[dim][1] * Q;
+			Input2Index += InputTensorInfo[dim][2] * Q;
+			Input3Index += InputTensorInfo[dim][3] * Q;
+			Offset = R;
+		}
+
+		BUFFER_TYPE Result = Input0[Input0Index];
+
+		 // Must correspond to EVariadicNumInput defined in NNIElementWiseVariadicCS.h
+		#if NUMINPUT >= 2
+			Result  = ELEMENTWISE_OP(Result, Input1[Input1Index]);
+		#endif
+		#if NUMINPUT >= 3
+			Result  = ELEMENTWISE_OP(Result, Input2[Input2Index]);
+		#endif
+		#if NUMINPUT >= 4
+			Result  = ELEMENTWISE_OP(Result, Input3[Input3Index]);
+		#endif
+
+		#if OUTPUTASINPUT
+			Result  = ELEMENTWISE_OP(Result, Output[Index]);
+		#endif
+
+		#if APPLYSCALE
+			Result  = (BUFFER_TYPE)(Result * Scale);
+		#endif
+
+		Output[Index] = Result;
+	}
+}
diff --git a/Engine/Plugins/Experimental/NNX/Source/NNIHlslShaders/Shaders/NNIHlslShadersGather.usf b/Engine/Plugins/Experimental/NNX/Source/NNIHlslShaders/Shaders/NNIHlslShadersGather.usf
new file mode 100644
--- /dev/null
+++ b/Engine/Plugins/Experimental/NNX/Source/NNIHlslShaders/Shaders/NNIHlslShadersGather.usf
@@ -0,0 +1,86 @@
+// Copyright Epic Games, Inc. All Rights Reserved.
+
+#include "/Engine/Public/Platform.ush"
+
+#if DATA_TYPE == 0
+#define DATA_BUFFER_TYPE float
+#else
+#define DATA_BUFFER_TYPE float
+#endif
+
+#if INDICES_TYPE == 0
+#define INDICES_WORK_TYPE int
+#define INDICES_BUFFER_TYPE int
+#define INDICES_BUFFER_TO_WORK_TYPE(x) x
+#else
+#define INDICES_WORK_TYPE int
+#define INDICES_BUFFER_TYPE int
+#define INDICES_BUFFER_TO_WORK_TYPE(x) x
+#endif
+
+int Axis;
+int OutputSize;
+int NumDataDimensions;
+int NumIndicesDimensions;
+int4 DataStride_IndicesStride_OutputStride[MAX_NUM_DIMENSIONS];
+float4 OneDivDataStride_OneDivIndicesStride_OneDivOutputStride[MAX_NUM_DIMENSIONS];
+
+Buffer<DATA_BUFFER_TYPE> Data;
+Buffer<INDICES_BUFFER_TYPE> Indices;
+RWBuffer<DATA_BUFFER_TYPE> Output;
+
+[numthreads(NUM_GROUP_THREADS, 1, 1)]
+void Gather(in const uint3 DispatchThreadID : SV_DispatchThreadID)
+{
+	if (DispatchThreadID.x >= OutputSize)
+	{
+		return;
+	}
+
+	// Compute the output index per dimenion
+	INDICES_WORK_TYPE OutputIndices[NUM_OUTPUT_DIMENSIONS];
+	INDICES_WORK_TYPE DispatchThreadIndex = (INDICES_WORK_TYPE)DispatchThreadID.x;
+	UNROLL
+	for (int i = 0; i < NUM_OUTPUT_DIMENSIONS; i++)
+	{
+		OutputIndices[i] = (INDICES_WORK_TYPE)(OneDivDataStride_OneDivIndicesStride_OneDivOutputStride[i].z * (float)DispatchThreadIndex);
+		DispatchThreadIndex -= OutputIndices[i] * (INDICES_WORK_TYPE)DataStride_IndicesStride_OutputStride[i].z;
+	}
+
+	// Split the output indices into data and indices indices 
+	INDICES_WORK_TYPE DataIndices[NUM_OUTPUT_DIMENSIONS];
+	INDICES_WORK_TYPE IndicesIndices[NUM_OUTPUT_DIMENSIONS];
+	int AxisPlusNumIndicesDimensions = Axis + NumIndicesDimensions;
+	for (int i = 0; i < Axis; i++)
+	{
+		DataIndices[i] = OutputIndices[i];
+	}
+	for (int i = 0; i < NumIndicesDimensions; i++)
+	{
+		IndicesIndices[i] = OutputIndices[Axis + i];
+	}
+	for (int i = Axis + 1; i < NumDataDimensions; i++)
+	{
+		DataIndices[i] = OutputIndices[AxisPlusNumIndicesDimensions + i];
+	}
+
+	// Compute the flattened indices index
+	INDICES_WORK_TYPE IndicesIndex = 0;
+	for (int i = 0; i < NumIndicesDimensions; i++)
+	{
+		IndicesIndex += IndicesIndices[i] * DataStride_IndicesStride_OutputStride[i].y;
+	}
+
+	// Look up and assign the missing data index
+	DataIndices[Axis] = INDICES_BUFFER_TO_WORK_TYPE(Indices[IndicesIndex]);
+
+	// Compute the flattened data index
+	INDICES_WORK_TYPE DataIndex = 0;
+	for (int i = 0; i < NumDataDimensions; i++)
+	{
+		DataIndex += DataIndices[i] * DataStride_IndicesStride_OutputStride[i].x;
+	}
+
+	// Write the result
+	Output[DispatchThreadID.x] = Data[DataIndex];
+}
\ No newline at end of file
diff --git a/Engine/Plugins/Experimental/NNX/Source/NNIHlslShaders/Shaders/NNIHlslShadersGemm.usf b/Engine/Plugins/Experimental/NNX/Source/NNIHlslShaders/Shaders/NNIHlslShadersGemm.usf
new file mode 100644
--- /dev/null
+++ b/Engine/Plugins/Experimental/NNX/Source/NNIHlslShaders/Shaders/NNIHlslShadersGemm.usf
@@ -0,0 +1,447 @@
+// Copyright Epic Games, Inc. All Rights Reserved.
+
+#include "/Engine/Public/Platform.ush"
+
+#define WORK_TYPE float
+#define BUFFER_TYPE float
+#define READ(x) x
+#define WRITE(x) x
+
+WORK_TYPE Alpha;
+WORK_TYPE Beta;
+int TransA;
+int TransB;
+uint M;
+uint N;
+uint K;
+uint MxK;
+uint KxN;
+uint MxN;
+uint CWidth;
+uint CHeight;
+
+WORK_TYPE CScalar;
+
+Buffer<BUFFER_TYPE> A;
+Buffer<BUFFER_TYPE> B;
+Buffer<BUFFER_TYPE> C;
+RWBuffer<BUFFER_TYPE> Y;
+
+// x: StackShapeA[0], ..., StackShapeA[NUM_STACK_DIMENSIONS - 1]
+// y: StackShapeB[0], ..., StackShapeB[NUM_STACK_DIMENSIONS - 1]
+// z: StackStrideA[0], ..., StackStrideA[NUM_STACK_DIMENSIONS - 1]
+// w: StackStrideB[0], ..., StackStrideB[NUM_STACK_DIMENSIONS - 1]
+uint4 StackShapeA_StackShapeB_StackStrideA_StackStrideB[MAX_NUM_STACK_DIMENSIONS];
+
+uint3 GetMatrixStackOffsets(const uint3 GroupID)
+{
+// Must correspond to FGemmNumDimensions defined in NNIHlslShadersGemmCS.h
+#if NUM_STACK_DIMENSIONS == 0
+	return uint3(0, 0, 0);
+#else
+	uint stackCoordA[NUM_STACK_DIMENSIONS];
+	uint stackCoordB[NUM_STACK_DIMENSIONS];
+
+	uint stackStrideOutput[NUM_STACK_DIMENSIONS];
+	stackStrideOutput[NUM_STACK_DIMENSIONS - 1] = 1;
+	for (int i = NUM_STACK_DIMENSIONS - 2; i >= 0; i--)
+	{
+		uint maxDim = max(StackShapeA_StackShapeB_StackStrideA_StackStrideB[i + 1][0], StackShapeA_StackShapeB_StackStrideA_StackStrideB[i + 1][1]);
+		stackStrideOutput[i] = maxDim * stackStrideOutput[i + 1];
+	}
+
+	uint idx = GroupID.z;
+	for (int i = 0; i < NUM_STACK_DIMENSIONS; i++)
+	{
+		uint coord = idx / stackStrideOutput[i];
+		stackCoordA[i] = coord % StackShapeA_StackShapeB_StackStrideA_StackStrideB[i][0];
+		stackCoordB[i] = coord % StackShapeA_StackShapeB_StackStrideA_StackStrideB[i][1];
+
+		idx = idx % stackStrideOutput[i];
+	}
+
+	uint stackOffsetA = 0;
+	uint stackOffsetB = 0;
+	for (int i = 0; i < NUM_STACK_DIMENSIONS; i++)
+	{
+		stackOffsetA += stackCoordA[i] * StackShapeA_StackShapeB_StackStrideA_StackStrideB[i][2];
+		stackOffsetB += stackCoordB[i] * StackShapeA_StackShapeB_StackStrideA_StackStrideB[i][3];
+	}
+
+	return uint3(
+		stackOffsetA	* MxK,
+		stackOffsetB	* KxN,
+		GroupID.z		* MxN);
+#endif
+}
+
+WORK_TYPE GetBetaTimesC(const uint3 DispatchThreadID)
+{
+// Must correspond to EGemmCScalar defined in NNIHlslShadersGemmCS.h
+#if C_SCALAR == 0
+	return Beta * READ(C[(DispatchThreadID.y % CHeight) * N + (DispatchThreadID.x % CWidth)]);
+#elif C_SCALAR == 1
+	return Beta * CScalar;
+#else
+	return 0;
+#endif
+}
+
+// Must correspond to EGemmAlgorithm defined in NNIHlslShadersGemmCS.h
+#if ALGORITHM < 4
+
+#if ALGORITHM == 0
+[numthreads(8, 8, 1)]
+#elif ALGORITHM == 1
+[numthreads(16, 16, 1)]
+#elif ALGORITHM == 2
+[numthreads(32, 32, 1)]
+#elif ALGORITHM == 3
+[numthreads(256, 1, 1)]
+#endif
+void Gemm(in const uint3 DispatchThreadID : SV_DispatchThreadID, in const uint3 GroupID : SV_GroupID)
+{
+	if (DispatchThreadID.y >= M || DispatchThreadID.x >= N)
+	{
+		return;
+	}
+	
+	uint AIdx;
+	uint AIdxIncrement;
+	if (TransA == 0)
+	{
+		AIdx = DispatchThreadID.y * K; // Index of the first element in the row
+		AIdxIncrement = 1; // The loop goes over the whole row
+	}
+	else
+	{
+		AIdx = DispatchThreadID.y; // Index of the first element in the column
+		AIdxIncrement = M; // The loop goes over the column, thus skip a full row
+	}
+	
+	uint BIdx;
+	uint BIdxIncrement;
+	if (TransB == 0)
+	{
+		BIdx = DispatchThreadID.x; // Index of the first element in the column
+		BIdxIncrement = N; // The loop goes over the column, thus skip a full row
+	}
+	else
+	{
+		BIdx = DispatchThreadID.x * K; // Index of the first element in the row
+		BIdxIncrement = 1; // The loop goes over the whole row
+	}
+
+	const uint3 StackOffsets = GetMatrixStackOffsets(GroupID);
+	AIdx += StackOffsets.x;
+	BIdx += StackOffsets.y;
+	
+	WORK_TYPE Result = 0.0;
+	for (int k = 0; k < K; k++)
+	{
+		Result += READ(A[AIdx]) * READ(B[BIdx]);
+		AIdx += AIdxIncrement;
+		BIdx += BIdxIncrement;
+	}
+	
+	Y[StackOffsets.z + DispatchThreadID.y * N + DispatchThreadID.x] = WRITE(Alpha * Result + GetBetaTimesC(DispatchThreadID));
+}
+
+#elif ALGORITHM < 7
+
+#if ALGORITHM == 4
+#define GROUP_SIZE 8
+#elif ALGORITHM == 5
+#define GROUP_SIZE 16
+#elif ALGORITHM == 6
+#define GROUP_SIZE 32
+#endif
+
+groupshared WORK_TYPE SharedMemoryA[GROUP_SIZE * GROUP_SIZE];
+groupshared WORK_TYPE SharedMemoryB[GROUP_SIZE * GROUP_SIZE];
+
+[numthreads(GROUP_SIZE, GROUP_SIZE, 1)]
+void Gemm(in const uint3 DispatchThreadID : SV_DispatchThreadID, in const uint3 GroupThreadID : SV_GroupThreadID, in const uint3 GroupID : SV_GroupID)
+{
+	// Calculate the number of steps each group needs to make to cover the full length of the matrix
+	uint NumGroupSteps = (uint)ceil((float)K / (float)GROUP_SIZE);
+	
+	// Compute some variables required at multiple places
+	uint SharedMemoryLineStartIdx = GroupThreadID.y * GROUP_SIZE;
+	uint SharedMemoryIdx = SharedMemoryLineStartIdx + GroupThreadID.x;
+	
+	// Compute the lookup start indices and the increment which needs to be applied for each sub block
+	uint AIdx;
+	uint AIdxIncrement;
+	if (TransA == 0)
+	{
+		AIdx = DispatchThreadID.y * K + GroupThreadID.x;
+		AIdxIncrement = GROUP_SIZE;
+	}
+	else
+	{
+		AIdx = GroupThreadID.x * M + DispatchThreadID.y;
+		AIdxIncrement = GROUP_SIZE * M;
+	}
+	
+	uint BIdx;
+	uint BIdxIncrement;
+	if (TransB == 0)
+	{
+		BIdx = GroupThreadID.y * N + DispatchThreadID.x;
+		BIdxIncrement = GROUP_SIZE * N;
+	}
+	else
+	{
+		BIdx = DispatchThreadID.x * K + GroupThreadID.y;
+		BIdxIncrement = GROUP_SIZE;
+	}
+
+	const uint3 StackOffsets = GetMatrixStackOffsets(GroupID);
+	AIdx += StackOffsets.x;
+	BIdx += StackOffsets.y;
+	
+	// Multiply and accumulate the blocks along one dimension of the matrices
+	WORK_TYPE Result = 0.0;
+	uint GroupIdx = 0;
+	for (int GroupStep = 0; GroupStep < NumGroupSteps; GroupStep++) 
+	{
+		// Only sync and update inidces starting with the second iteration
+		if(GroupStep > 0)
+		{
+			// Increment the indices
+			GroupIdx += GROUP_SIZE;
+			AIdx += AIdxIncrement;
+			BIdx += BIdxIncrement;
+			
+			// Make sure computation is done in the whole group before overwriting shared memory
+			GroupMemoryBarrierWithGroupSync();
+		}
+		
+		// Load the shared memory after applying the transpose and set overflow elements to zero
+		WORK_TYPE Temp = 0;
+		if (DispatchThreadID.y < M && GroupIdx + GroupThreadID.x < K)
+		{
+			Temp = READ(A[AIdx]);
+		}
+		SharedMemoryA[SharedMemoryIdx] = Temp;
+
+		Temp = 0;
+		if (DispatchThreadID.x < N && GroupIdx + GroupThreadID.y < K)
+		{
+			Temp = READ(B[BIdx]);
+		}
+		SharedMemoryB[SharedMemoryIdx] = Temp;
+		
+		// Make sure all memory has been loaded before running the multiplication
+		GroupMemoryBarrierWithGroupSync();
+		
+		// Multiply the groups and accumulate the result
+		// This works also at matrix borders, as overflow elements are set to zero
+		uint SharedMemoryBIdx = GroupThreadID.x;
+		for (int kGroup = 0; kGroup < GROUP_SIZE; kGroup++)
+		{
+			Result += SharedMemoryA[SharedMemoryLineStartIdx + kGroup] * SharedMemoryB[SharedMemoryBIdx];
+			SharedMemoryBIdx += GROUP_SIZE;
+		}
+	}
+	
+	// Only now return out of bound threads as they were needed to load shared memory above
+	if (DispatchThreadID.y >= M || DispatchThreadID.x >= N)
+	{
+		return;
+	}
+	
+	Y[StackOffsets.z + DispatchThreadID.y * N + DispatchThreadID.x] = WRITE(Alpha * Result + GetBetaTimesC(DispatchThreadID));
+}
+
+#else
+
+#if ALGORITHM == 7
+#define GROUP_SIZE_X 		16
+#define GROUP_SIZE_Y 		1
+#define LOAD_PER_THREAD 	16
+#elif ALGORITHM == 8
+#define GROUP_SIZE_X 		16
+#define GROUP_SIZE_Y 		2
+#define LOAD_PER_THREAD 	8
+#elif ALGORITHM == 9
+#define GROUP_SIZE_X 		32
+#define GROUP_SIZE_Y 		1
+#define LOAD_PER_THREAD 	32
+#elif ALGORITHM == 10
+#define GROUP_SIZE_X 		32
+#define GROUP_SIZE_Y 		2
+#define LOAD_PER_THREAD 	16
+#elif ALGORITHM == 11
+#define GROUP_SIZE_X 		32
+#define GROUP_SIZE_Y 		4
+#define LOAD_PER_THREAD 	8
+#elif ALGORITHM == 12
+#define GROUP_SIZE_X 		64
+#define GROUP_SIZE_Y 		2
+#define LOAD_PER_THREAD 	32
+#elif ALGORITHM == 13
+#define GROUP_SIZE_X 		64
+#define GROUP_SIZE_Y 		4
+#define LOAD_PER_THREAD 	16
+#elif ALGORITHM == 14
+#define GROUP_SIZE_X 		64
+#define GROUP_SIZE_Y 		8
+#define LOAD_PER_THREAD 	8
+#endif
+
+// Use the larger of group sizes to define the shared memory size
+// Threads will do multiple loads over the smaller dimension to fill the memory completely
+groupshared WORK_TYPE SharedMemoryA[GROUP_SIZE_X * GROUP_SIZE_X];
+groupshared WORK_TYPE SharedMemoryB[GROUP_SIZE_X * GROUP_SIZE_X];
+
+[numthreads(GROUP_SIZE_X, GROUP_SIZE_Y, 1)]
+void Gemm(in const uint3 DispatchThreadID : SV_DispatchThreadID, in const uint3 GroupThreadID : SV_GroupThreadID, in const uint3 GroupID : SV_GroupID)
+{
+	// Initialize the result
+	WORK_TYPE Results[LOAD_PER_THREAD];
+	for (int LoadIdx = 0; LoadIdx < LOAD_PER_THREAD; LoadIdx++)
+	{
+        Results[LoadIdx] = 0.0f;
+    }
+	
+	// Calculate the number of steps each group needs to make to cover the full length of the matrix
+	uint NumGroupSteps = (uint)ceil((float)K / (float)GROUP_SIZE_X);
+	
+	// Compute some variables required at multiple places
+	uint SharedMemoryTransposedLineIdx = GroupThreadID.x * GROUP_SIZE_X;
+	uint SharedMemoryStartIdx = GroupThreadID.y * GROUP_SIZE_X + GroupThreadID.x;
+	uint SharedIdxIncrement = GROUP_SIZE_Y * GROUP_SIZE_X;
+	uint GroupStartRow = GroupID.y * GROUP_SIZE_X;
+	uint GroupLoadOffset = GroupThreadID.y * LOAD_PER_THREAD;
+	
+	// Compute the lookup start indices and the increment which needs to be applied for each sub block
+	uint AIdx;
+	uint AIdxOuterIncrement; // Increment added in the group loop
+	uint AIdxInnerIncrement; // Increment added in the load loop
+	if (TransA == 0)
+	{
+		AIdx = (GroupStartRow + GroupThreadID.y) * K + GroupThreadID.x;
+		AIdxOuterIncrement = GROUP_SIZE_X;
+		AIdxInnerIncrement = GROUP_SIZE_Y * K;
+	}
+	else
+	{
+		AIdx = GroupThreadID.x * M + (GroupStartRow + GroupThreadID.y);
+		AIdxOuterIncrement = GROUP_SIZE_X * M;
+		AIdxInnerIncrement = GROUP_SIZE_Y;
+	}
+	
+	uint BIdx;
+	uint BIdxOuterIncrement; // Increment added in the group loop
+	uint BIdxInnerIncrement; // Increment added in the load loop
+	if (TransB == 0)
+	{
+		BIdx = GroupThreadID.y * N + DispatchThreadID.x;
+		BIdxOuterIncrement = GROUP_SIZE_X * N;
+		BIdxInnerIncrement = GROUP_SIZE_Y * N;
+	}
+	else
+	{
+		BIdx = DispatchThreadID.x * K + GroupThreadID.y;
+		BIdxOuterIncrement = GROUP_SIZE_X;
+		BIdxInnerIncrement = GROUP_SIZE_Y;
+	}
+
+	const uint3 StackOffsets = GetMatrixStackOffsets(GroupID);
+	AIdx += StackOffsets.x;
+	BIdx += StackOffsets.y;
+
+	// Multiply and accumulate the blocks
+	uint GroupIdx = 0;
+	for (uint GroupStep = 0; GroupStep < NumGroupSteps; GroupStep++) 
+	{
+		// Prepare some shared variables needed inside this loop
+		uint SharedIdx = SharedMemoryStartIdx;
+		uint AIdxOffset = 0;
+		uint BIdxOffset = 0;
+		
+		uint GroupRow = GroupThreadID.y;
+		uint ThreadIDY = GroupStartRow + GroupRow;
+		
+		// Only sync and update inidces starting with the second iteration
+		if (GroupStep > 0)
+		{
+			// increment indices
+			GroupIdx += GROUP_SIZE_X;
+			AIdx += AIdxOuterIncrement;
+			BIdx += BIdxOuterIncrement;
+			
+			// Make sure computation is done in the whole group before overwriting shared memory
+			GroupMemoryBarrierWithGroupSync();
+		}
+		
+		for (int LoadIdx = 0; LoadIdx < LOAD_PER_THREAD; LoadIdx++)
+		{
+			// Load the shared memory after applying the transpose and set overflow elements to zero
+			WORK_TYPE Temp = 0;
+			if (ThreadIDY < M && GroupIdx + GroupThreadID.x < K)
+			{
+				Temp = READ(A[AIdx + AIdxOffset]);
+			}
+			SharedMemoryA[SharedIdx] = Temp;
+			
+			Temp = 0;
+			if (GroupIdx + GroupRow < K && DispatchThreadID.x < N)
+			{
+				Temp = READ(B[BIdx + BIdxOffset]);
+			}
+			SharedMemoryB[SharedIdx] = Temp;
+			
+			// Increment the thread y coordinate virtually by one group size
+			GroupRow += GROUP_SIZE_Y;
+			ThreadIDY += GROUP_SIZE_Y;
+			
+			// Increment the shared index pointer to advance to the next subblock
+			SharedIdx += SharedIdxIncrement;
+			AIdxOffset += AIdxInnerIncrement;
+			BIdxOffset += BIdxInnerIncrement;
+		}
+		
+		// Make sure all memory has been loaded before running the multiplication
+		GroupMemoryBarrierWithGroupSync();
+		
+		// multiply and accumulate
+		uint StartB = GroupLoadOffset;
+		for (int kGroup = 0; kGroup < GROUP_SIZE_X; kGroup++)
+		{
+			WORK_TYPE AVal = SharedMemoryA[SharedMemoryTransposedLineIdx + kGroup];
+			for (int LoadIdx = 0; LoadIdx < LOAD_PER_THREAD; LoadIdx++)
+			{
+				Results[LoadIdx] += AVal * SharedMemoryB[StartB + LoadIdx];
+			}
+			StartB += GROUP_SIZE_X;
+		}
+	}
+	
+	// Only now return out of bound threads as they were needed to load shared memory above
+	uint Row = GroupStartRow + GroupThreadID.x;
+	if (Row >= M)
+	{
+		return;
+	}
+	
+	// Save back the results
+	uint RowIdx = Row * N;
+	uint ColumnStart = GroupID.x * GROUP_SIZE_X + GroupLoadOffset;
+	for (int LoadIdx = 0; LoadIdx < LOAD_PER_THREAD; LoadIdx++)
+	{
+		uint Column = ColumnStart + LoadIdx;
+		
+		// Return as soon as the column is out of bound as further loops will be outside too
+		if (Column >= N)
+		{
+			return;
+		}
+		
+		Y[StackOffsets.z + RowIdx + Column] = WRITE(Alpha * Results[LoadIdx] + GetBetaTimesC(uint3(Column, Row, 0)));
+    }
+}
+
+#endif
diff --git a/Engine/Plugins/Experimental/NNX/Source/NNIHlslShaders/Shaders/NNIHlslShadersIdentity.usf b/Engine/Plugins/Experimental/NNX/Source/NNIHlslShaders/Shaders/NNIHlslShadersIdentity.usf
new file mode 100644
--- /dev/null
+++ b/Engine/Plugins/Experimental/NNX/Source/NNIHlslShaders/Shaders/NNIHlslShadersIdentity.usf
@@ -0,0 +1,18 @@
+// Copyright Epic Games, Inc. All Rights Reserved.
+
+#include "/Engine/Public/Platform.ush"
+
+Buffer<float> Input;
+RWBuffer<float> Output;
+uint Num;
+
+[numthreads(THREADGROUP_SIZE_X, 1, 1)]
+void Identity(in const uint3 DispatchThreadID : SV_DispatchThreadID)
+{
+	const uint Index = DispatchThreadID.x;
+	
+	if (Index < Num)
+	{
+		Output[Index] = Input[Index];
+	}
+}
diff --git a/Engine/Plugins/Experimental/NNX/Source/NNXHlslShaders/Internal/NNIHlslShadersConvCS.h b/Engine/Plugins/Experimental/NNX/Source/NNXHlslShaders/Internal/NNIHlslShadersConvCS.h
deleted file mode 100644
diff --git a/Engine/Plugins/Experimental/NNX/Source/NNXHlslShaders/Internal/NNIHlslShadersConvTransposeCS.h b/Engine/Plugins/Experimental/NNX/Source/NNXHlslShaders/Internal/NNIHlslShadersConvTransposeCS.h
deleted file mode 100644
diff --git a/Engine/Plugins/Experimental/NNX/Source/NNXHlslShaders/Internal/NNIHlslShadersElementWiseBinaryCS.h b/Engine/Plugins/Experimental/NNX/Source/NNXHlslShaders/Internal/NNIHlslShadersElementWiseBinaryCS.h
deleted file mode 100644
diff --git a/Engine/Plugins/Experimental/NNX/Source/NNXHlslShaders/Internal/NNIHlslShadersElementWiseUnaryCS.h b/Engine/Plugins/Experimental/NNX/Source/NNXHlslShaders/Internal/NNIHlslShadersElementWiseUnaryCS.h
deleted file mode 100644
diff --git a/Engine/Plugins/Experimental/NNX/Source/NNXHlslShaders/Internal/NNIHlslShadersElementWiseVariadicCS.h b/Engine/Plugins/Experimental/NNX/Source/NNXHlslShaders/Internal/NNIHlslShadersElementWiseVariadicCS.h
deleted file mode 100644
diff --git a/Engine/Plugins/Experimental/NNX/Source/NNXHlslShaders/Internal/NNIHlslShadersGatherCS.h b/Engine/Plugins/Experimental/NNX/Source/NNXHlslShaders/Internal/NNIHlslShadersGatherCS.h
deleted file mode 100644
diff --git a/Engine/Plugins/Experimental/NNX/Source/NNXHlslShaders/Internal/NNIHlslShadersGemmCS.h b/Engine/Plugins/Experimental/NNX/Source/NNXHlslShaders/Internal/NNIHlslShadersGemmCS.h
deleted file mode 100644
diff --git a/Engine/Plugins/Experimental/NNX/Source/NNXHlslShaders/NNXHlslShaders.Build.cs b/Engine/Plugins/Experimental/NNX/Source/NNXHlslShaders/NNXHlslShaders.Build.cs
deleted file mode 100644
diff --git a/Engine/Plugins/Experimental/NNX/Source/NNXHlslShaders/Private/NNIHlslShadersConvCS.cpp b/Engine/Plugins/Experimental/NNX/Source/NNXHlslShaders/Private/NNIHlslShadersConvCS.cpp
deleted file mode 100644
diff --git a/Engine/Plugins/Experimental/NNX/Source/NNXHlslShaders/Private/NNIHlslShadersConvTransposeCS.cpp b/Engine/Plugins/Experimental/NNX/Source/NNXHlslShaders/Private/NNIHlslShadersConvTransposeCS.cpp
deleted file mode 100644
diff --git a/Engine/Plugins/Experimental/NNX/Source/NNXHlslShaders/Private/NNIHlslShadersElementWiseBinaryCS.cpp b/Engine/Plugins/Experimental/NNX/Source/NNXHlslShaders/Private/NNIHlslShadersElementWiseBinaryCS.cpp
deleted file mode 100644
diff --git a/Engine/Plugins/Experimental/NNX/Source/NNXHlslShaders/Private/NNIHlslShadersElementWiseUnaryCS.cpp b/Engine/Plugins/Experimental/NNX/Source/NNXHlslShaders/Private/NNIHlslShadersElementWiseUnaryCS.cpp
deleted file mode 100644
diff --git a/Engine/Plugins/Experimental/NNX/Source/NNXHlslShaders/Private/NNIHlslShadersElementWiseVariadicCS.cpp b/Engine/Plugins/Experimental/NNX/Source/NNXHlslShaders/Private/NNIHlslShadersElementWiseVariadicCS.cpp
deleted file mode 100644
diff --git a/Engine/Plugins/Experimental/NNX/Source/NNXHlslShaders/Private/NNIHlslShadersGatherCS.cpp b/Engine/Plugins/Experimental/NNX/Source/NNXHlslShaders/Private/NNIHlslShadersGatherCS.cpp
deleted file mode 100644
diff --git a/Engine/Plugins/Experimental/NNX/Source/NNXHlslShaders/Private/NNIHlslShadersGemmCS.cpp b/Engine/Plugins/Experimental/NNX/Source/NNXHlslShaders/Private/NNIHlslShadersGemmCS.cpp
deleted file mode 100644
diff --git a/Engine/Plugins/Experimental/NNX/Source/NNXHlslShaders/Private/NNXHlslModule.cpp b/Engine/Plugins/Experimental/NNX/Source/NNXHlslShaders/Private/NNXHlslModule.cpp
deleted file mode 100644
diff --git a/Engine/Plugins/Experimental/NNX/Source/NNXHlslShaders/Public/NNXShaderParameters.h b/Engine/Plugins/Experimental/NNX/Source/NNXHlslShaders/Public/NNXShaderParameters.h
deleted file mode 100644
diff --git a/Engine/Plugins/Experimental/NNX/Source/NNXHlslShaders/Shaders/BroadcastHelper.ush b/Engine/Plugins/Experimental/NNX/Source/NNXHlslShaders/Shaders/BroadcastHelper.ush
deleted file mode 100644
diff --git a/Engine/Plugins/Experimental/NNX/Source/NNXHlslShaders/Shaders/ConvOp.usf b/Engine/Plugins/Experimental/NNX/Source/NNXHlslShaders/Shaders/ConvOp.usf
deleted file mode 100644
diff --git a/Engine/Plugins/Experimental/NNX/Source/NNXHlslShaders/Shaders/ConvTransposeOp.usf b/Engine/Plugins/Experimental/NNX/Source/NNXHlslShaders/Shaders/ConvTransposeOp.usf
deleted file mode 100644
diff --git a/Engine/Plugins/Experimental/NNX/Source/NNXHlslShaders/Shaders/ElementWiseBinaryOp.usf b/Engine/Plugins/Experimental/NNX/Source/NNXHlslShaders/Shaders/ElementWiseBinaryOp.usf
deleted file mode 100644
diff --git a/Engine/Plugins/Experimental/NNX/Source/NNXHlslShaders/Shaders/ElementWiseOp.usf b/Engine/Plugins/Experimental/NNX/Source/NNXHlslShaders/Shaders/ElementWiseOp.usf
deleted file mode 100644
diff --git a/Engine/Plugins/Experimental/NNX/Source/NNXHlslShaders/Shaders/ElementWiseVariadicOp.usf b/Engine/Plugins/Experimental/NNX/Source/NNXHlslShaders/Shaders/ElementWiseVariadicOp.usf
deleted file mode 100644
diff --git a/Engine/Plugins/Experimental/NNX/Source/NNXHlslShaders/Shaders/GatherOp.usf b/Engine/Plugins/Experimental/NNX/Source/NNXHlslShaders/Shaders/GatherOp.usf
deleted file mode 100644
diff --git a/Engine/Plugins/Experimental/NNX/Source/NNXHlslShaders/Shaders/GemmOp.usf b/Engine/Plugins/Experimental/NNX/Source/NNXHlslShaders/Shaders/GemmOp.usf
deleted file mode 100644
diff --git a/Engine/Plugins/Experimental/NNX/Source/NNXHlslShaders/Shaders/IdentityOp.usf b/Engine/Plugins/Experimental/NNX/Source/NNXHlslShaders/Shaders/IdentityOp.usf
deleted file mode 100644
diff --git a/Engine/Plugins/Experimental/NNX/Source/NNXRuntimeRDG/NNXRuntimeRDG.Build.cs b/Engine/Plugins/Experimental/NNX/Source/NNXRuntimeRDG/NNXRuntimeRDG.Build.cs
--- a/Engine/Plugins/Experimental/NNX/Source/NNXRuntimeRDG/NNXRuntimeRDG.Build.cs
+++ b/Engine/Plugins/Experimental/NNX/Source/NNXRuntimeRDG/NNXRuntimeRDG.Build.cs
@@ -24,7 +24,7 @@ public class NNXRuntimeRDG : ModuleRules
         {
             "NNXCore",
 			"NNXUtils",
-			"NNXHlslShaders",
+			"NNIHlslShaders",
             "RHI"
         });
 
