NNI - Removing unused files and renaming unary operator

diff --git a/Engine/Plugins/Experimental/NNX/Source/NNIHlslShaders/Private/NNIHlslShadersElementWiseUnaryCS.cpp b/Engine/Plugins/Experimental/NNX/Source/NNIHlslShaders/Private/NNIHlslShadersElementWiseUnaryCS.cpp
--- a/Engine/Plugins/Experimental/NNX/Source/NNIHlslShaders/Private/NNIHlslShadersElementWiseUnaryCS.cpp
+++ b/Engine/Plugins/Experimental/NNX/Source/NNIHlslShaders/Private/NNIHlslShadersElementWiseUnaryCS.cpp
@@ -78,5 +78,5 @@ namespace UE::NNIHlslShaders::Internal
 		return OpFunc;
 	}
 
-	IMPLEMENT_GLOBAL_SHADER(TElementWiseUnaryCS, "/NNI/NNIHlslShadersElementWise.usf", "ElementWise", SF_Compute);
+	IMPLEMENT_GLOBAL_SHADER(TElementWiseUnaryCS, "/NNI/NNIHlslShadersElementWiseUnary.usf", "ElementWiseUnary", SF_Compute);
 } // UE::NNIHlslShaders::Internal
diff --git a/Engine/Plugins/Experimental/NNX/Source/NNIHlslShaders/Shaders/NNIHlslShadersElementWise.usf b/Engine/Plugins/Experimental/NNX/Source/NNIHlslShaders/Shaders/NNIHlslShadersElementWise.usf
deleted file mode 100644
diff --git a/Engine/Plugins/Experimental/NNX/Source/NNIHlslShaders/Shaders/NNIHlslShadersElementWiseUnary.usf b/Engine/Plugins/Experimental/NNX/Source/NNIHlslShaders/Shaders/NNIHlslShadersElementWiseUnary.usf
new file mode 100644
--- /dev/null
+++ b/Engine/Plugins/Experimental/NNX/Source/NNIHlslShaders/Shaders/NNIHlslShadersElementWiseUnary.usf
@@ -0,0 +1,118 @@
+// Copyright Epic Games, Inc. All Rights Reserved.
+
+#include "/Engine/Public/Platform.ush"
+
+Buffer<float> Input;
+RWBuffer<float> Output;
+uint Num;
+uint ThreadCountX;
+float Alpha;
+float Beta;
+float Gamma;
+
+//see [REDACTED_PROFILE_URL]/onnx/blob/main/docs/Operators.md#Relu
+float relu(float x)
+{
+	return max(0.0f, x);
+}
+
+//see [REDACTED_PROFILE_URL]/onnx/blob/main/docs/Operators.md#Selu
+float selu(float x)
+{
+	float yNegOrZero = Gamma * (Alpha * exp(x) - Alpha);
+	float yPos = Gamma * x;
+	return (x > 0.0f) ? yPos : yNegOrZero;
+}
+
+//see [REDACTED_PROFILE_URL]/onnx/blob/main/docs/Operators.md#sigmoid
+float sigmoid(float x)
+{
+	return 1.0f / (1.0f + exp(-x));
+}
+
+//see [REDACTED_PROFILE_URL]/onnx/blob/main/docs/Operators.md#softplus
+float softplus(float x)
+{
+	return log(exp(x) + 1.0f);
+}
+
+//see [REDACTED_PROFILE_URL]/onnx/blob/main/docs/Operators.md#softsign
+float softsign(float x)
+{
+	return x/(1.0f+abs(x));
+}
+
+//see [REDACTED_PROFILE_URL]/onnx/blob/main/docs/Operators.md#acosh
+float acosh(float x)
+{
+	//https://mathworld.wolfram.com/InverseHyperbolicCosine.html
+	const static float NaN = 0.0f / 0.0f;
+	float yAboveOne = log( x + sqrt( x + 1.0f ) + sqrt( x - 1.0f ) );
+	return (x == 1.0f) ? 0.0f : (x >= 1.0f) ? yAboveOne : NaN;
+}
+
+//see [REDACTED_PROFILE_URL]/onnx/blob/main/docs/Operators.md#asinh
+float asinh(float x)
+{
+	//https://mathworld.wolfram.com/InverseHyperbolicSine.html
+	return log( x + sqrt( 1 + ( x * x ) ) );
+}
+
+//see [REDACTED_PROFILE_URL]/onnx/blob/main/docs/Operators.md#atanh
+float atanh(float x)
+{
+	//https://mathworld.wolfram.com/InverseHyperbolicTangent.html
+	return 0.5f * ( log( 1 + x) - log ( 1 - x) );
+}
+
+//see [REDACTED_PROFILE_URL]/onnx/blob/main/docs/Operators.md#elu
+float elu(float x)
+{
+    float yNeg = Alpha * (exp(x) - 1.0f);
+	float yPosOrZero = x;
+	return (x >= 0.0f) ? yPosOrZero : yNeg;
+}
+
+//see [REDACTED_PROFILE_URL]/onnx/blob/main/docs/Operators.md#erf
+float erf(float x)
+{
+	//https://aapt.scitation.org/doi/abs/10.1119/1.15018?journalCode=ajp
+	float a = 167.0f/148.0f;
+	float b = 11.0f/109.0f;
+	float x3 = x*x*x;
+	return tanh( a*x + b*x3 );
+}
+
+float hardSigmoid(float x, float alpha, float beta)
+{
+	return max(0.0f, min(1.0f, alpha * x + beta));
+}
+
+//see [REDACTED_PROFILE_URL]/onnx/blob/main/docs/Operators.md#hardSigmoid
+float hardSigmoid(float x)
+{
+	return hardSigmoid(x, Alpha, Beta);
+}
+
+//see [REDACTED_PROFILE_URL]/onnx/blob/main/docs/Operators.md#hardSwish
+float hardSwish(float x)
+{
+    return x * hardSigmoid(x, 1.0f/6.0f, 0.5f);
+}
+
+//see [REDACTED_PROFILE_URL]/onnx/blob/main/docs/Operators.md#leakyRelu
+float leakyRelu(float x)
+{
+    return (x >= 0.0f) ? x : Alpha * x;
+}
+
+
+[numthreads(THREADGROUP_SIZE_X, 1, 1)]
+void ElementWiseUnary(in const uint3 DispatchThreadID : SV_DispatchThreadID)
+{
+	const uint Index = DispatchThreadID.y * ThreadCountX + DispatchThreadID.x;
+	if (Index < Num)
+	{
+		Output[Index] = ELEMENTWISE_OP(Input[Index]);
+	}
+}
diff --git a/Engine/Plugins/Experimental/NNX/Source/NNIHlslShaders/Shaders/NNIHlslShadersIdentity.usf b/Engine/Plugins/Experimental/NNX/Source/NNIHlslShaders/Shaders/NNIHlslShadersIdentity.usf
deleted file mode 100644
