NNI - Cleaning NNIProfiling

diff --git a/Engine/Plugins/Experimental/NNX/NNX.uplugin b/Engine/Plugins/Experimental/NNX/NNX.uplugin
--- a/Engine/Plugins/Experimental/NNX/NNX.uplugin
+++ b/Engine/Plugins/Experimental/NNX/NNX.uplugin
@@ -38,7 +38,7 @@
 			"PlatformAllowList": [ "Win64", "Linux", "Mac" ]
 		},
 		{
-			"Name": "NNXProfiling",
+			"Name": "NNIProfiling",
 			"Type": "RuntimeAndProgram",
 			"LoadingPhase": "Default",
 			"PlatformAllowList": [ "Win64", "Linux", "Mac" ]
diff --git a/Engine/Plugins/Experimental/NNX/Source/NNIProfiling/Internal/NNIProfilingStatistics.h b/Engine/Plugins/Experimental/NNX/Source/NNIProfiling/Internal/NNIProfilingStatistics.h
new file mode 100644
--- /dev/null
+++ b/Engine/Plugins/Experimental/NNX/Source/NNIProfiling/Internal/NNIProfilingStatistics.h
@@ -0,0 +1,44 @@
+// Copyright Epic Games, Inc. All Rights Reserved.
+
+#pragma once
+
+#include "CoreMinimal.h"
+#include "Containers/CircularBuffer.h"
+
+namespace UE::NNIProfiling::Internal
+{
+	struct NNIPROFILING_API FStatistics
+	{
+		const uint32 NumberSamples;
+		const float Average;
+		const float StdDev;
+		const float Min;
+		const float Max;
+
+		FStatistics(const int InNumberSamples, const float InAverage, const float InStdDev, const float InMin, const float InMax);
+	};
+
+	class NNIPROFILING_API FStatisticsEstimator
+	{
+	public:
+
+		FStatisticsEstimator(const uint32 InSizeRollingWindow = 1024);
+
+		void StoreSample(const float InRunTime);
+		void ResetStats();
+
+		float GetLastSample() const;
+		FStatistics GetStats() const;
+
+	private:
+
+		const float EpsilonFloat;
+		uint32 BufferIdx;
+		bool bIsBufferFull;
+		float LastSample;
+		TCircularBuffer<float> DataBuffer;
+
+		float CalculateMean() const;
+		float CalculateStdDev(const float InMean) const;
+	};
+} // UE::NNIProfiling::Internal
\ No newline at end of file
diff --git a/Engine/Plugins/Experimental/NNX/Source/NNIProfiling/Internal/NNIProfilingTimer.h b/Engine/Plugins/Experimental/NNX/Source/NNIProfiling/Internal/NNIProfilingTimer.h
new file mode 100644
--- /dev/null
+++ b/Engine/Plugins/Experimental/NNX/Source/NNIProfiling/Internal/NNIProfilingTimer.h
@@ -0,0 +1,24 @@
+// Copyright Epic Games, Inc. All Rights Reserved.
+
+#pragma once
+
+#include "CoreMinimal.h"
+
+namespace UE::NNIProfiling::Internal
+{
+	class NNIPROFILING_API FTimer
+	{
+	public:
+
+		void Tic();
+
+		/**
+		 * Time in milliseconds, but with nanosecond accuracy (e.g., 1.234567 msec).
+		 */
+		double Toc() const;
+
+	private:
+
+		double TimeStart;
+	};
+} // UE::NNIProfiling::Internal
\ No newline at end of file
diff --git a/Engine/Plugins/Experimental/NNX/Source/NNIProfiling/NNIProfiling.Build.cs b/Engine/Plugins/Experimental/NNX/Source/NNIProfiling/NNIProfiling.Build.cs
new file mode 100644
--- /dev/null
+++ b/Engine/Plugins/Experimental/NNX/Source/NNIProfiling/NNIProfiling.Build.cs
@@ -0,0 +1,21 @@
+// Copyright Epic Games, Inc. All Rights Reserved.
+
+using UnrealBuildTool;
+
+public class NNIProfiling : ModuleRules
+{
+	public NNIProfiling(ReadOnlyTargetRules Target) : base(Target)
+	{
+		PCHUsage = PCHUsageMode.UseExplicitOrSharedPCHs;
+
+		PublicDependencyModuleNames.AddRange
+			(
+			new string[] {
+				"Core",
+				"Engine",
+			}
+		);
+
+		PublicDefinitions.Add("WITH_NNI_STATS");
+	}
+}
diff --git a/Engine/Plugins/Experimental/NNX/Source/NNIProfiling/Private/NNIProfilingModule.cpp b/Engine/Plugins/Experimental/NNX/Source/NNIProfiling/Private/NNIProfilingModule.cpp
new file mode 100644
--- /dev/null
+++ b/Engine/Plugins/Experimental/NNX/Source/NNIProfiling/Private/NNIProfilingModule.cpp
@@ -0,0 +1,21 @@
+// Copyright Epic Games, Inc. All Rights Reserved.
+
+#include "CoreMinimal.h"
+#include "Modules/ModuleManager.h"
+
+class FNNIProfilingModule : public IModuleInterface
+{
+public:
+
+	virtual void StartupModule() override
+	{
+		
+	}
+
+	virtual void ShutdownModule() override
+	{
+
+	}
+};
+
+IMPLEMENT_MODULE(FNNIProfilingModule, NNIProfiling);
diff --git a/Engine/Plugins/Experimental/NNX/Source/NNIProfiling/Private/NNIProfilingStatistics.cpp b/Engine/Plugins/Experimental/NNX/Source/NNIProfiling/Private/NNIProfilingStatistics.cpp
new file mode 100644
--- /dev/null
+++ b/Engine/Plugins/Experimental/NNX/Source/NNIProfiling/Private/NNIProfilingStatistics.cpp
@@ -0,0 +1,112 @@
+// Copyright Epic Games, Inc. All Rights Reserved.
+
+#include "NNIProfilingStatistics.h"
+#include "GenericPlatform/GenericPlatformMath.h"
+
+namespace UE::NNIProfiling::Internal
+{
+	void FStatisticsEstimator_CalculateMinMax(float& OutMin, float& OutMax, const TCircularBuffer<float>& InDataBuffer, const bool bInIsBufferFull, const uint32 InBufferIdx)
+	{
+		const uint32 PtrEnd = (bInIsBufferFull) ? InDataBuffer.Capacity() : InBufferIdx;
+
+		OutMin = FLT_MAX;
+		OutMax = FLT_MIN;
+		for (uint32 Index = 0; Index < PtrEnd; ++Index)
+		{
+			const float CurrentElement = InDataBuffer[Index];
+			if (CurrentElement < OutMin)
+			{
+				OutMin = CurrentElement;
+			}
+			if (CurrentElement > OutMax)
+			{
+				OutMax = CurrentElement;
+			}
+		}
+	}
+
+	FStatistics::FStatistics(const int InNumberSamples, const float InAverage, const float InStdDev, const float InMin, const float InMax)
+		: NumberSamples(InNumberSamples)
+		, Average(InAverage)
+		, StdDev(InStdDev)
+		, Min(InMin)
+		, Max(InMax)
+	{
+
+	}
+
+	FStatisticsEstimator::FStatisticsEstimator(const uint32 InSizeRollingWindow)
+		: EpsilonFloat(0.000001f)
+		, BufferIdx(0)
+		, bIsBufferFull(false)
+		, LastSample(0.f)
+		, DataBuffer(InSizeRollingWindow, 0)
+	{
+
+	}
+
+	void FStatisticsEstimator::StoreSample(const float InRunTime)
+	{
+		LastSample = InRunTime;
+
+		DataBuffer[BufferIdx] = LastSample;
+
+		if (BufferIdx == DataBuffer.Capacity() && !bIsBufferFull)
+		{
+			bIsBufferFull = true;
+		}
+		BufferIdx = DataBuffer.GetNextIndex(BufferIdx);
+	}
+
+	void FStatisticsEstimator::ResetStats()
+	{
+		for (uint32 Index = 0; Index < DataBuffer.Capacity(); ++Index)
+		{
+			DataBuffer[Index] = 0;
+		}
+		bIsBufferFull = false;
+		BufferIdx = 0;
+		LastSample = 0;
+	}
+
+	float FStatisticsEstimator::GetLastSample() const
+	{
+		return LastSample;
+	}
+
+	FStatistics FStatisticsEstimator::GetStats() const
+	{
+		const uint32 count = (bIsBufferFull) ? DataBuffer.Capacity() : BufferIdx;
+		const float Mean = CalculateMean();
+		const float StdDev = CalculateStdDev(Mean);
+		float Min, Max;
+		FStatisticsEstimator_CalculateMinMax(Min, Max, DataBuffer, bIsBufferFull, BufferIdx);
+		// Return aggregated data
+		return FStatistics(count, Mean, StdDev, Min, Max);
+	}
+
+	float FStatisticsEstimator::CalculateMean() const
+	{
+		uint32 PtrEnd = (bIsBufferFull) ? DataBuffer.Capacity() : BufferIdx;
+		float AccumulatorMean = 0;
+		for (uint32 Index = 0; Index < PtrEnd; ++Index)
+		{
+			AccumulatorMean += DataBuffer[Index];
+		}
+		// Return mean
+		return AccumulatorMean / (static_cast<float>(PtrEnd) + EpsilonFloat);
+	}
+
+	float FStatisticsEstimator::CalculateStdDev(const float InMean) const
+	{
+		const uint32 PtrEnd = (bIsBufferFull) ? DataBuffer.Capacity() : BufferIdx;
+
+		float AccumulatorStdDev = 0;
+		for (uint32 Index = 0; Index < PtrEnd; ++Index)
+		{
+			AccumulatorStdDev += (DataBuffer[Index] - InMean) * (DataBuffer[Index] - InMean);
+		}
+		// Return std deviation
+		return FMath::Sqrt(AccumulatorStdDev / (static_cast<float>(PtrEnd) + EpsilonFloat));
+	}
+} // UE::NNIProfiling::Internal
\ No newline at end of file
diff --git a/Engine/Plugins/Experimental/NNX/Source/NNIProfiling/Private/NNIProfilingTimer.cpp b/Engine/Plugins/Experimental/NNX/Source/NNIProfiling/Private/NNIProfilingTimer.cpp
new file mode 100644
--- /dev/null
+++ b/Engine/Plugins/Experimental/NNX/Source/NNIProfiling/Private/NNIProfilingTimer.cpp
@@ -0,0 +1,17 @@
+// Copyright Epic Games, Inc. All Rights Reserved.
+
+#include "NNIProfilingTimer.h"
+#include "HAL/PlatformTime.h"
+
+namespace UE::NNIProfiling::Internal
+{
+	void FTimer::Tic()
+	{
+		TimeStart = { FPlatformTime::Seconds() };
+	}
+
+	double FTimer::Toc() const
+	{
+		return (FPlatformTime::Seconds() - TimeStart) * 1e3;
+	}
+} // UE::NNIProfiling::Internal
\ No newline at end of file
diff --git a/Engine/Plugins/Experimental/NNX/Source/NNXProfiling/NNXProfiling.Build.cs b/Engine/Plugins/Experimental/NNX/Source/NNXProfiling/NNXProfiling.Build.cs
deleted file mode 100644
diff --git a/Engine/Plugins/Experimental/NNX/Source/NNXProfiling/Private/Module.cpp b/Engine/Plugins/Experimental/NNX/Source/NNXProfiling/Private/Module.cpp
deleted file mode 100644
diff --git a/Engine/Plugins/Experimental/NNX/Source/NNXProfiling/Private/Module.h b/Engine/Plugins/Experimental/NNX/Source/NNXProfiling/Private/Module.h
deleted file mode 100644
diff --git a/Engine/Plugins/Experimental/NNX/Source/NNXProfiling/Private/NeuralStatistics.cpp b/Engine/Plugins/Experimental/NNX/Source/NNXProfiling/Private/NeuralStatistics.cpp
deleted file mode 100644
diff --git a/Engine/Plugins/Experimental/NNX/Source/NNXProfiling/Private/NeuralTimer.cpp b/Engine/Plugins/Experimental/NNX/Source/NNXProfiling/Private/NeuralTimer.cpp
deleted file mode 100644
diff --git a/Engine/Plugins/Experimental/NNX/Source/NNXProfiling/Public/NeuralStatistics.h b/Engine/Plugins/Experimental/NNX/Source/NNXProfiling/Public/NeuralStatistics.h
deleted file mode 100644
diff --git a/Engine/Plugins/Experimental/NNX/Source/NNXProfiling/Public/NeuralTimer.h b/Engine/Plugins/Experimental/NNX/Source/NNXProfiling/Public/NeuralTimer.h
deleted file mode 100644
diff --git a/Engine/Plugins/Experimental/NNX/Source/NNXRuntimeCPU/NNXRuntimeCPU.Build.cs b/Engine/Plugins/Experimental/NNX/Source/NNXRuntimeCPU/NNXRuntimeCPU.Build.cs
--- a/Engine/Plugins/Experimental/NNX/Source/NNXRuntimeCPU/NNXRuntimeCPU.Build.cs
+++ b/Engine/Plugins/Experimental/NNX/Source/NNXRuntimeCPU/NNXRuntimeCPU.Build.cs
@@ -19,7 +19,7 @@ public class NNXRuntimeCPU : ModuleRules
 				"Core",
 				"Engine",
 				"NNXCore",
-				"NNXProfiling",
+				"NNIProfiling",
 				"ORTHelper",
 			}
 		);
diff --git a/Engine/Plugins/Experimental/NNX/Source/NNXRuntimeCPU/Private/NNXRuntimeCPU.cpp b/Engine/Plugins/Experimental/NNX/Source/NNXRuntimeCPU/Private/NNXRuntimeCPU.cpp
--- a/Engine/Plugins/Experimental/NNX/Source/NNXRuntimeCPU/Private/NNXRuntimeCPU.cpp
+++ b/Engine/Plugins/Experimental/NNX/Source/NNXRuntimeCPU/Private/NNXRuntimeCPU.cpp
@@ -2,7 +2,7 @@
 #include "NNXRuntimeCPU.h"
 #include "NNXRuntimeCPUUtils.h"
 #include "NNXModelOptimizer.h"
-#include "NeuralTimer.h"
+#include "NNIProfilingTimer.h"
 #include "RedirectCoutAndCerrToUeLog.h"
 
 using namespace NNX;
@@ -228,7 +228,7 @@ int FMLInferenceModelCPU::Run(TArrayView<const NNX::FMLTensorBinding> InInputBin
 		return -1;
 	}
 
-	FNeuralTimer RunTimer;
+	UE::NNIProfiling::Internal::FTimer RunTimer;
 	RunTimer.Tic();
 
 	if (!bHasRun)
@@ -268,12 +268,12 @@ float FMLInferenceModelCPU::GetLastRunTimeMSec() const
 	return RunStatisticsEstimator.GetLastSample();
 }
 
-FNeuralStatistics FMLInferenceModelCPU::GetRunStatistics() const
+UE::NNIProfiling::Internal::FStatistics FMLInferenceModelCPU::GetRunStatistics() const
 {
 	return RunStatisticsEstimator.GetStats();
 }
 
-FNeuralStatistics FMLInferenceModelCPU::GetInputMemoryTransferStats() const
+UE::NNIProfiling::Internal::FStatistics FMLInferenceModelCPU::GetInputMemoryTransferStats() const
 {
 	return InputTransferStatisticsEstimator.GetStats();
 }
diff --git a/Engine/Plugins/Experimental/NNX/Source/NNXRuntimeCPU/Public/NNXRuntimeCPU.h b/Engine/Plugins/Experimental/NNX/Source/NNXRuntimeCPU/Public/NNXRuntimeCPU.h
--- a/Engine/Plugins/Experimental/NNX/Source/NNXRuntimeCPU/Public/NNXRuntimeCPU.h
+++ b/Engine/Plugins/Experimental/NNX/Source/NNXRuntimeCPU/Public/NNXRuntimeCPU.h
@@ -6,7 +6,7 @@
 #include "NNXCore.h"
 #include "NNXRuntime.h"
 #include "NNXInferenceModel.h"
-#include "NeuralStatistics.h"
+#include "NNIProfilingStatistics.h"
 
 #include "NNXThirdPartyWarningDisabler.h"
 NNX_THIRD_PARTY_INCLUDES_START
@@ -86,8 +86,8 @@ namespace NNX
 		virtual int Run(TArrayView<const NNX::FMLTensorBinding> InInputBindingTensors, TArrayView<const NNX::FMLTensorBinding> OutOutputBindingTensors);
 
 		float GetLastRunTimeMSec() const;
-		FNeuralStatistics GetRunStatistics() const;
-		FNeuralStatistics GetInputMemoryTransferStats() const;
+		UE::NNIProfiling::Internal::FStatistics GetRunStatistics() const;
+		UE::NNIProfiling::Internal::FStatistics GetInputMemoryTransferStats() const;
 		void ResetStats();
 
 	protected:
@@ -118,8 +118,8 @@ namespace NNX
 		/**
 		 * Statistics-related members used for GetLastRunTimeMSec(), GetRunStatistics(), GetInputMemoryTransferStats(), ResetStats().
 		 */
-		FNeuralStatisticsEstimator RunStatisticsEstimator;
-		FNeuralStatisticsEstimator InputTransferStatisticsEstimator;
+		UE::NNIProfiling::Internal::FStatisticsEstimator RunStatisticsEstimator;
+		UE::NNIProfiling::Internal::FStatisticsEstimator InputTransferStatisticsEstimator;
 
 	};
 }
\ No newline at end of file
diff --git a/Engine/Plugins/Experimental/NNX/Source/NNXRuntimeORT/NNXRuntimeORT.Build.cs b/Engine/Plugins/Experimental/NNX/Source/NNXRuntimeORT/NNXRuntimeORT.Build.cs
--- a/Engine/Plugins/Experimental/NNX/Source/NNXRuntimeORT/NNXRuntimeORT.Build.cs
+++ b/Engine/Plugins/Experimental/NNX/Source/NNXRuntimeORT/NNXRuntimeORT.Build.cs
@@ -22,7 +22,7 @@ public class NNXRuntimeORT : ModuleRules
 				"Core",
 				"Engine",
 				"NNXCore",
-				"NNXProfiling",
+				"NNIProfiling",
 			}
 		);
 
diff --git a/Engine/Plugins/Experimental/NNX/Source/NNXRuntimeORT/Private/NNXRuntimeORT.cpp b/Engine/Plugins/Experimental/NNX/Source/NNXRuntimeORT/Private/NNXRuntimeORT.cpp
--- a/Engine/Plugins/Experimental/NNX/Source/NNXRuntimeORT/Private/NNXRuntimeORT.cpp
+++ b/Engine/Plugins/Experimental/NNX/Source/NNXRuntimeORT/Private/NNXRuntimeORT.cpp
@@ -4,7 +4,7 @@
 #include "CoreGlobals.h"
 #include "Misc/ConfigCacheIni.h"
 
-#include "NeuralTimer.h"
+#include "NNIProfilingTimer.h"
 #include "NNXRuntimeORTUtils.h"
 #include "NNXModelOptimizer.h"
 #include "RedirectCoutAndCerrToUeLog.h"
@@ -301,7 +301,7 @@ int FMLInferenceModelORT::Run(TArrayView<const NNX::FMLTensorBinding> InInputBin
 		return -1;
 	}
 
-	FNeuralTimer RunTimer;
+	UE::NNIProfiling::Internal::FTimer RunTimer;
 	RunTimer.Tic();
 
 	if (!bHasRun)
@@ -340,12 +340,12 @@ float FMLInferenceModelORT::GetLastRunTimeMSec() const
 	return RunStatisticsEstimator.GetLastSample();
 }
 
-FNeuralStatistics FMLInferenceModelORT::GetRunStatistics() const
+UE::NNIProfiling::Internal::FStatistics FMLInferenceModelORT::GetRunStatistics() const
 {
 	return RunStatisticsEstimator.GetStats();
 }
 
-FNeuralStatistics FMLInferenceModelORT::GetInputMemoryTransferStats() const
+UE::NNIProfiling::Internal::FStatistics FMLInferenceModelORT::GetInputMemoryTransferStats() const
 {
 	return InputTransferStatisticsEstimator.GetStats();
 }
diff --git a/Engine/Plugins/Experimental/NNX/Source/NNXRuntimeORT/Private/NNXRuntimeORT.h b/Engine/Plugins/Experimental/NNX/Source/NNXRuntimeORT/Private/NNXRuntimeORT.h
--- a/Engine/Plugins/Experimental/NNX/Source/NNXRuntimeORT/Private/NNXRuntimeORT.h
+++ b/Engine/Plugins/Experimental/NNX/Source/NNXRuntimeORT/Private/NNXRuntimeORT.h
@@ -6,7 +6,7 @@
 #include "NNXCore.h"
 #include "NNXRuntime.h"
 #include "NNXInferenceModel.h"
-#include "NeuralStatistics.h"
+#include "NNIProfilingStatistics.h"
 #include "NNXThirdPartyWarningDisabler.h"
 
 NNX_THIRD_PARTY_INCLUDES_START
@@ -185,8 +185,8 @@ namespace NNX
 		virtual int Run(TArrayView<const FMLTensorBinding> InInputBindingTensors, TArrayView<const FMLTensorBinding> OutOutputBindingTensors);
 
 		float GetLastRunTimeMSec() const;
-		FNeuralStatistics GetRunStatistics() const;
-		FNeuralStatistics GetInputMemoryTransferStats() const;
+		UE::NNIProfiling::Internal::FStatistics GetRunStatistics() const;
+		UE::NNIProfiling::Internal::FStatistics GetInputMemoryTransferStats() const;
 		void ResetStats();
 
 	protected:
@@ -216,8 +216,8 @@ namespace NNX
 		/**
 		 * Statistics-related members used for GetLastRunTimeMSec(), GetRunStatistics(), GetInputMemoryTransferStats(), ResetStats().
 		 */
-		FNeuralStatisticsEstimator RunStatisticsEstimator;
-		FNeuralStatisticsEstimator InputTransferStatisticsEstimator;
+		UE::NNIProfiling::Internal::FStatisticsEstimator RunStatisticsEstimator;
+		UE::NNIProfiling::Internal::FStatisticsEstimator InputTransferStatisticsEstimator;
 	};
 
 	class FMLInferenceModelORTCpu : public NNX::FMLInferenceModelORT
