NNX: Preparation toward multi layer support and shape inference

Problems:
* NNX does not have multi layer support atm
* NNX does not  have varaible input shape support atm

Solutions:
Prepare to both of the above by laying the fundation so those change can be worked on incrementally. In that directions:
* Support intermediate tensors in model definition
* Add symoblic shape support to serialized model
* Define a symbolic tensor shape VS a concrete tensor shape for the runtime
* Add tensor shapes to model validation

Tests:
* All NNX test green on win64 editor

diff --git a/Engine/Plugins/Experimental/NNX/Source/NNIHlslShaders/Private/NNIHlslShadersGatherCS.cpp b/Engine/Plugins/Experimental/NNX/Source/NNIHlslShaders/Private/NNIHlslShadersGatherCS.cpp
--- a/Engine/Plugins/Experimental/NNX/Source/NNIHlslShaders/Private/NNIHlslShadersGatherCS.cpp
+++ b/Engine/Plugins/Experimental/NNX/Source/NNIHlslShaders/Private/NNIHlslShadersGatherCS.cpp
@@ -25,17 +25,17 @@ namespace UE::NNIHlslShaders::Internal
 		Parameters.OutputSize = 1;
 		for (int32 i = 0; i < Axis; i++, index++)
 		{
-			OutputShape[index] = Data.Shape[i];
+			OutputShape[index] = Data.Shape.Data[i];
 			Parameters.OutputSize *= OutputShape[index];
 		}
 		for (int32 i = 0; i < Indices.Shape.Num(); i++, index++)
 		{
-			OutputShape[index] = Indices.Shape[i];
+			OutputShape[index] = Indices.Shape.Data[i];
 			Parameters.OutputSize *= OutputShape[index];
 		}
 		for (int32 i = (Axis + 1); i < Data.Shape.Num(); i++, index++)
 		{
-			OutputShape[index] = Data.Shape[i];
+			OutputShape[index] = Data.Shape.Data[i];
 			Parameters.OutputSize *= OutputShape[index];
 		}
 
@@ -47,7 +47,7 @@ namespace UE::NNIHlslShaders::Internal
 		{
 			Parameters.DataStride_IndicesStride_OutputStride[i].X = DataStride;
 			Parameters.OneDivDataStride_OneDivIndicesStride_OneDivOutputStride[i].X = 1.0 / (float)DataStride;
-			DataStride *= Data.Shape[i];
+			DataStride *= Data.Shape.Data[i];
 		}
 
 		int32 IndicesStride = 1;
@@ -55,7 +55,7 @@ namespace UE::NNIHlslShaders::Internal
 		{
 			Parameters.DataStride_IndicesStride_OutputStride[i].Y = IndicesStride;
 			Parameters.OneDivDataStride_OneDivIndicesStride_OneDivOutputStride[i].Y = 1.0 / (float)IndicesStride;
-			IndicesStride *= Indices.Shape[i];
+			IndicesStride *= Indices.Shape.Data[i];
 		}
 
 		int32 OutputStride = 1;
diff --git a/Engine/Plugins/Experimental/NNX/Source/NNIHlslShaders/Private/NNIHlslShadersGemmCS.cpp b/Engine/Plugins/Experimental/NNX/Source/NNIHlslShaders/Private/NNIHlslShadersGemmCS.cpp
--- a/Engine/Plugins/Experimental/NNX/Source/NNIHlslShaders/Private/NNIHlslShadersGemmCS.cpp
+++ b/Engine/Plugins/Experimental/NNX/Source/NNIHlslShaders/Private/NNIHlslShadersGemmCS.cpp
@@ -35,8 +35,8 @@ namespace UE::NNIHlslShaders::Internal
 			const int32 NumDimToCheck = FMath::Min(NumStackDimensionsA, NumStackDimensionsB);
 			for (int32 i = 0; i < NumDimToCheck; i++)
 			{
-				const uint32 VolumeA = InputA.Shape[NumStackDimensionsA - 1 - i];
-				const uint32 VolumeB = InputB.Shape[NumStackDimensionsB - 1 - i];
+				const uint32 VolumeA = InputA.Shape.Data[NumStackDimensionsA - 1 - i];
+				const uint32 VolumeB = InputB.Shape.Data[NumStackDimensionsB - 1 - i];
 
 				check(VolumeA == 1 || VolumeB == 1 || VolumeA == VolumeB);
 			}
@@ -45,10 +45,10 @@ namespace UE::NNIHlslShaders::Internal
 		const bool IsVectorA = InputA.Shape.Num() == 1;
 		const bool IsVectorB = InputB.Shape.Num() == 1;
 
-		Result.M = IsVectorA ? 1 : InputA.Shape[InputA.Shape.Num() - 2];
-		Result.N = IsVectorB ? 1 : InputB.Shape[InputB.Shape.Num() - 1];
-		Result.K = InputA.Shape[IsVectorA ? 0 : InputA.Shape.Num() - 1];
-		check(InputB.Shape[IsVectorB ? 0 : InputB.Shape.Num() - 2] == Result.K);
+		Result.M = IsVectorA ? 1 : InputA.Shape.Data[InputA.Shape.Num() - 2];
+		Result.N = IsVectorB ? 1 : InputB.Shape.Data[InputB.Shape.Num() - 1];
+		Result.K = InputA.Shape.Data[IsVectorA ? 0 : InputA.Shape.Num() - 1];
+		check(InputB.Shape.Data[IsVectorB ? 0 : InputB.Shape.Num() - 2] == Result.K);
 
 		Result.StackShapeA.Init(1, NumStackDimensions);
 		Result.StackShapeB.Init(1, NumStackDimensions);
@@ -58,8 +58,8 @@ namespace UE::NNIHlslShaders::Internal
 			int32 IdxA = InputA.Shape.Num() - 3 - i;
 			int32 IdxB = InputB.Shape.Num() - 3 - i;
 
-			Result.StackShapeA[Result.StackShapeA.Num() - 1 - i] = IdxA >= 0 ? InputA.Shape[IdxA] : 1;
-			Result.StackShapeB[Result.StackShapeB.Num() - 1 - i] = IdxB >= 0 ? InputB.Shape[IdxB] : 1;
+			Result.StackShapeA[Result.StackShapeA.Num() - 1 - i] = IdxA >= 0 ? InputA.Shape.Data[IdxA] : 1;
+			Result.StackShapeB[Result.StackShapeB.Num() - 1 - i] = IdxB >= 0 ? InputB.Shape.Data[IdxB] : 1;
 		}
 
 		Result.StackStrideA.Init(1, NumStackDimensions);
@@ -85,10 +85,10 @@ namespace UE::NNIHlslShaders::Internal
 		check(InputA.Shape.Num() == 2);
 		check(InputB.Shape.Num() == 2);
 
-		uint32 M = TransA != 0 ? InputA.Shape[1] : InputA.Shape[0];
-		uint32 K = TransA != 0 ? InputA.Shape[0] : InputA.Shape[1];
-		uint32 N = TransB != 0 ? InputB.Shape[0] : InputB.Shape[1];
-		check(K == (TransB != 0 ? InputB.Shape[1] : InputB.Shape[0]));
+		uint32 M = TransA != 0 ? InputA.Shape.Data[1] : InputA.Shape.Data[0];
+		uint32 K = TransA != 0 ? InputA.Shape.Data[0] : InputA.Shape.Data[1];
+		uint32 N = TransB != 0 ? InputB.Shape.Data[0] : InputB.Shape.Data[1];
+		check(K == (TransB != 0 ? InputB.Shape.Data[1] : InputB.Shape.Data[0]));
 
 		Parameters.Alpha = Alpha;
 		Parameters.Beta = Beta;
@@ -100,8 +100,8 @@ namespace UE::NNIHlslShaders::Internal
 		Parameters.MxK = M * K;
 		Parameters.KxN = K * N;
 		Parameters.MxN = M * N;
-		Parameters.CWidth = InputC.Shape.Num() == 0 ? 0 : InputC.Shape[InputC.Shape.Num() - 1];;
-		Parameters.CHeight = InputC.Shape.Num() == 0 ? 0 : InputC.Shape.Num() == 1 ? 1 : InputC.Shape[InputC.Shape.Num() - 2];
+		Parameters.CWidth = InputC.Shape.Num() == 0 ? 0 : InputC.Shape.Data[InputC.Shape.Num() - 1];;
+		Parameters.CHeight = InputC.Shape.Num() == 0 ? 0 : InputC.Shape.Num() == 1 ? 1 : InputC.Shape.Data[InputC.Shape.Num() - 2];
 		Parameters.CScalar = CScalar;
 	}
 
diff --git a/Engine/Plugins/Experimental/NNX/Source/NNXCore/Private/NNXModelOptimizerInterface.h b/Engine/Plugins/Experimental/NNX/Source/NNXCore/Private/NNXModelOptimizerInterface.h
--- a/Engine/Plugins/Experimental/NNX/Source/NNXCore/Private/NNXModelOptimizerInterface.h
+++ b/Engine/Plugins/Experimental/NNX/Source/NNXCore/Private/NNXModelOptimizerInterface.h
@@ -16,7 +16,7 @@ class IModelValidator
 public:
 	virtual ~IModelValidator() = default;
 	virtual FString GetName() const = 0;
-	virtual bool ValidateModel(const FNNIModelRaw& InputModel) const = 0;
+	virtual bool ValidateModel(const FNNIModelRaw& InputModel, const FOptimizerOptionsMap& Options) const = 0;
 };
 
 /** Interface class for NNX model optimizer pass */
diff --git a/Engine/Plugins/Experimental/NNX/Source/NNXCore/Public/NNXRuntime.h b/Engine/Plugins/Experimental/NNX/Source/NNXCore/Public/NNXRuntime.h
--- a/Engine/Plugins/Experimental/NNX/Source/NNXCore/Public/NNXRuntime.h
+++ b/Engine/Plugins/Experimental/NNX/Source/NNXCore/Public/NNXRuntime.h
@@ -171,8 +171,9 @@ protected:
 		: Type(InType)
 	{}
 
-	TArray<FMLTensorDesc >	InputTensors;
-	TArray<FMLTensorDesc >	OutputTensors;
+	TArray<FMLTensorDesc>	InputTensors;
+	TArray<FMLTensorDesc>	OutputTensors;
+	TArray<FMLTensorDesc>	AllTensors;
 	EMLInferenceModelType	Type;
 };
 
diff --git a/Engine/Plugins/Experimental/NNX/Source/NNXCore/Public/NNXRuntimeFormat.h b/Engine/Plugins/Experimental/NNX/Source/NNXCore/Public/NNXRuntimeFormat.h
--- a/Engine/Plugins/Experimental/NNX/Source/NNXCore/Public/NNXRuntimeFormat.h
+++ b/Engine/Plugins/Experimental/NNX/Source/NNXCore/Public/NNXRuntimeFormat.h
@@ -95,7 +95,7 @@ struct FMLFormatTensorDesc
 	FString Name;
 
 	UPROPERTY(VisibleAnywhere, Category = "Neural Network Inference")
-	TArray<uint32> Shape;
+	TArray<int32> Shape;
 
 	UPROPERTY(VisibleAnywhere, Category = "Neural Network Inference")
 	EMLFormatTensorType	Type = EMLFormatTensorType::None;
diff --git a/Engine/Plugins/Experimental/NNX/Source/NNXCore/Public/NNXTypes.h b/Engine/Plugins/Experimental/NNX/Source/NNXCore/Public/NNXTypes.h
--- a/Engine/Plugins/Experimental/NNX/Source/NNXCore/Public/NNXTypes.h
+++ b/Engine/Plugins/Experimental/NNX/Source/NNXCore/Public/NNXTypes.h
@@ -189,6 +189,78 @@ private:
 namespace NNX
 {
 
+struct FSymbolicTensorShape
+{
+	constexpr static int32	MaxRank = 8;
+	TArray<int32, TInlineAllocator<MaxRank>> Data;
+
+	int32 Num() const
+	{
+		return Data.Num();
+	}
+
+	static FSymbolicTensorShape Make(TConstArrayView<const int32> Data)
+	{
+		if (Data.Num() > MaxRank)
+		{
+			//UE_LOG(LogNNX, Warning, TEXT("Cannot create symbolic tensor shape, input is rank %d while max rank is %d"), Data.Num(), MaxRank);
+			return {};
+		}
+
+		FSymbolicTensorShape Shape;
+		
+		Shape.Data.Append(Data);
+		return Shape;
+	}
+	
+	bool IsConcrete()
+	{
+		for (int32 i = 0; i < Data.Num(); ++i)
+		{
+			if (Data[i] < 0)
+			{
+				return false;
+			}
+		}
+		return true;
+	}
+};
+
+struct FConcreteTensorShape
+{
+	TArray<uint32, TInlineAllocator<FSymbolicTensorShape::MaxRank>> Data;
+
+	int32 Num() const
+	{
+		return Data.Num();
+	}
+	
+	static FConcreteTensorShape Make(const FSymbolicTensorShape& SymbolicShape)
+	{
+		FConcreteTensorShape ConcreteShape;
+		for (int32 i = 0; i < SymbolicShape.Data.Num(); ++i)
+		{
+			int32 Dim = SymbolicShape.Data[i];
+			ConcreteShape.Data.Add(Dim < 0 ? 1 : Dim);
+		}
+		return ConcreteShape;
+	}
+
+	static FConcreteTensorShape Make(TConstArrayView<const uint32> Data)
+	{
+		if (Data.Num() > FSymbolicTensorShape::MaxRank)
+		{
+			//UE_LOG(LogNNX, Warning, TEXT("Cannot create concrete tensor shape, input is rank %d while max rank is %d"), Data.Num(), MaxRank);
+			return {};
+		}
+
+		FConcreteTensorShape Shape;
+
+		Shape.Data.Append(Data);
+		return Shape;
+	}
+};
+
 /**
  *  Tensor descriptor
  */
@@ -197,21 +269,14 @@ struct FMLTensorDesc
 	constexpr static int32	MaxTensorDimension = 5;
 
 	FString					Name;
-	TArray<uint32, TInlineAllocator<MaxTensorDimension>> Shape;	//!< Shape
-	uint64					DataSize;							//!< Size of data in bytes
+	FConcreteTensorShape	Shape;
+	uint64					DataSize;//!< Size of data in bytes
 	EMLTensorDataType		DataType;
 
 	/** Make tensor descriptor */
-	static FMLTensorDesc Make(const FString& Name, TArrayView<const uint32> Shape, EMLTensorDataType DataType, uint64 DataSize = 0)
+	static FMLTensorDesc Make(const FString& Name, const FConcreteTensorShape& Shape, EMLTensorDataType DataType, uint64 DataSize = 0)
 	{
-		check(Shape.Num() <= FMLTensorDesc::MaxTensorDimension);
-		if (Shape.Num() > FMLTensorDesc::MaxTensorDimension)
-		{
-			//UE_LOG(LogNNX, Warning, TEXT("Unsupported tensor dimension:%d"), Shape.Num());
-			return {};
-		}
-
-		return {Name, TArray<uint32, TInlineAllocator<MaxTensorDimension>>{Shape}, DataSize, DataType};
+		return {Name, Shape, DataSize, DataType};
 	}
 
 	/** Check if descriptor is valid */
@@ -267,7 +332,7 @@ inline size_t GetTensorDataTypeSizeInBytes(EMLTensorDataType InType)
 /** Check if descriptor is valid */
 inline bool FMLTensorDesc::IsValid() const
 {
-	return DataType != EMLTensorDataType::None && Shape.Num() > 0 && Shape.Num() <= MaxTensorDimension; //&& Volume() >= 1;
+	return DataType != EMLTensorDataType::None && Shape.Data.Num() > 0 && Shape.Data.Num() <= MaxTensorDimension; //&& Volume() >= 1;
 }
 
 /** Return size of one element in bytes */
@@ -281,9 +346,9 @@ inline int32 FMLTensorDesc::Volume() const
 {
 	int32 Result = 1;
 
-	for (int32 Idx = 0; Idx < Shape.Num(); ++Idx)
+	for (int32 Idx = 0; Idx < Shape.Data.Num(); ++Idx)
 	{
-		Result *= Shape[Idx];
+		Result *= Shape.Data[Idx];
 	}
 
 	return Result;
diff --git a/Engine/Plugins/Experimental/NNX/Source/NNXQA/Private/NNXQAParametricTest.cpp b/Engine/Plugins/Experimental/NNX/Source/NNXQA/Private/NNXQAParametricTest.cpp
--- a/Engine/Plugins/Experimental/NNX/Source/NNXQA/Private/NNXQAParametricTest.cpp
+++ b/Engine/Plugins/Experimental/NNX/Source/NNXQA/Private/NNXQAParametricTest.cpp
@@ -117,16 +117,12 @@ namespace Test
 			TestSetup.Tags = TestTarget.Tags;
 		}
 
-		static TArray<uint32> GetShapeFromJsonArray(const TArray<int32>& JsonShape)
+		static FConcreteTensorShape GetConcreteShapeFromJsonArray(const TArray<int32>& JsonShape)
 		{
-			TArray<uint32> Shape;
-			Shape.Reserve(JsonShape.Num());
-			for (int32 dim : JsonShape)
-			{
-				Shape.Emplace(dim >= 0 ? dim : 1);
-			}
-
-			return Shape;
+			FSymbolicTensorShape SymbolicShape = FSymbolicTensorShape::Make(JsonShape);
+			check(SymbolicShape.IsConcrete());
+			FConcreteTensorShape ConcreteShape = FConcreteTensorShape::Make(SymbolicShape);
+			return ConcreteShape;
 		}
 
 		static EMLTensorDataType GetTensorTypeFromJson(const FString& TypeName, EMLTensorDataType DefaultValue)
@@ -148,24 +144,24 @@ namespace Test
 			uint32 i = 0;
 			for (auto&& Tensor : TestDataset.Inputs)
 			{
-				TArray<uint32> shape = GetShapeFromJsonArray(Tensor.Shape);
+				FConcreteTensorShape Shape = GetConcreteShapeFromJsonArray(Tensor.Shape);
 				EMLTensorDataType TensorType = GetTensorTypeFromJson(Tensor.Type, DefaultInputType);
-				TestSetup.Inputs.Emplace(FMLTensorDesc::Make(FString::Printf(TEXT("in%d"), i++), shape, TensorType));
+				TestSetup.Inputs.Emplace(FMLTensorDesc::Make(FString::Printf(TEXT("in%d"), i++), Shape, TensorType));
 			}
 
 			i = 0;
 			for (auto&& Tensor : TestDataset.Outputs)
 			{
-				TArray<uint32> shape = GetShapeFromJsonArray(Tensor.Shape);
+				FConcreteTensorShape Shape = GetConcreteShapeFromJsonArray(Tensor.Shape);
 				EMLTensorDataType TensorType = GetTensorTypeFromJson(Tensor.Type, DefaultOutputType);
-				TestSetup.Outputs.Emplace(FMLTensorDesc::Make(FString::Printf(TEXT("output%d"), i++), shape, TensorType));
+				TestSetup.Outputs.Emplace(FMLTensorDesc::Make(FString::Printf(TEXT("output%d"), i++), Shape, TensorType));
 			}
 			//If output is not defined it is the first input shape.
 			if (TestDataset.Outputs.Num() == 0 && TestDataset.Inputs.Num() > 0)
 			{
-				TArray<uint32> shape = GetShapeFromJsonArray(TestDataset.Inputs[0].Shape);
+				FConcreteTensorShape Shape = GetConcreteShapeFromJsonArray(TestDataset.Inputs[0].Shape);
 				EMLTensorDataType TensorType = GetTensorTypeFromJson(TestDataset.Inputs[0].Type, DefaultOutputType);
-				TestSetup.Outputs.Emplace(FMLTensorDesc::Make(TEXT("output"), shape, TensorType));
+				TestSetup.Outputs.Emplace(FMLTensorDesc::Make(TEXT("output"), Shape, TensorType));
 			}
 		}
 
diff --git a/Engine/Plugins/Experimental/NNX/Source/NNXQA/Private/NNXQAUtils.cpp b/Engine/Plugins/Experimental/NNX/Source/NNXQA/Private/NNXQAUtils.cpp
--- a/Engine/Plugins/Experimental/NNX/Source/NNXQA/Private/NNXQAUtils.cpp
+++ b/Engine/Plugins/Experimental/NNX/Source/NNXQA/Private/NNXQAUtils.cpp
@@ -86,7 +86,7 @@ namespace Test
 
 	FString FMLTensorDescToString(const FMLTensorDesc& desc)
 	{
-		TArrayView<const uint32> Shape(desc.Shape);
+		TArrayView<const uint32> Shape(desc.Shape.Data);
 
 		FString TensorDesc;
 		TensorDesc.Reserve(50);
@@ -253,7 +253,7 @@ namespace Test
 		check(RefTensorDesc.Shape.Num() <= RefTensorDesc.MaxTensorDimension);
 		for (int32 i = 0; i < RefTensorDesc.Shape.Num(); ++i)
 		{
-			bTensorDescMatch &= (RefTensorDesc.Shape[i] == OtherTensorDesc.Shape[i]);
+			bTensorDescMatch &= (RefTensorDesc.Shape.Data[i] == OtherTensorDesc.Shape.Data[i]);
 		}
 
 		if (!bTensorDescMatch)
diff --git a/Engine/Plugins/Experimental/NNX/Source/NNXRuntimeCPU/Private/NNXRuntimeCPU.cpp b/Engine/Plugins/Experimental/NNX/Source/NNXRuntimeCPU/Private/NNXRuntimeCPU.cpp
--- a/Engine/Plugins/Experimental/NNX/Source/NNXRuntimeCPU/Private/NNXRuntimeCPU.cpp
+++ b/Engine/Plugins/Experimental/NNX/Source/NNXRuntimeCPU/Private/NNXRuntimeCPU.cpp
@@ -188,7 +188,7 @@ bool FMLInferenceModelCPU::ConfigureTensors(const bool InIsInput)
 		std::pair<EMLTensorDataType, uint64> TypeAndSize = TranslateTensorTypeORTToNNI(ONNXTensorElementDataTypeEnum);
 		CurrentTensorDescriptor.DataType = TypeAndSize.first;
 
-		CurrentTensorDescriptor.Shape.SetNumUninitialized(CurrentTensorInfo.GetShape().size());
+		CurrentTensorDescriptor.Shape.Data.SetNumUninitialized(CurrentTensorInfo.GetShape().size());
 		uint8 Index = 0;
 		uint64 TensorNumberElements = 1;
 		for (const int64_t CurrentTensorSize : CurrentTensorInfo.GetShape())
@@ -196,14 +196,14 @@ bool FMLInferenceModelCPU::ConfigureTensors(const bool InIsInput)
 			// Negative (variable) dimensions not implemented yet
 			if (CurrentTensorSize < 0)
 			{
-				CurrentTensorDescriptor.Shape[Index] = 1;
+				CurrentTensorDescriptor.Shape.Data[Index] = 1;
 				UE_LOG(LogNNX, Display,
 					TEXT("Negative (i.e., variable) dimensions not allowed yet, hard-coded to 1. Let us know if you really need variable dimensions."
 						" Keep in mind that fixed sizes might allow additional optimizations and speedup of the network during Run()."));
 			}
 			else
 			{
-				CurrentTensorDescriptor.Shape[Index] = CurrentTensorSize;
+				CurrentTensorDescriptor.Shape.Data[Index] = CurrentTensorSize;
 				TensorNumberElements *= CurrentTensorSize;
 			}
 			++Index;
diff --git a/Engine/Plugins/Experimental/NNX/Source/NNXRuntimeCPU/Public/NNXRuntimeCPUUtils.h b/Engine/Plugins/Experimental/NNX/Source/NNXRuntimeCPU/Public/NNXRuntimeCPUUtils.h
--- a/Engine/Plugins/Experimental/NNX/Source/NNXRuntimeCPU/Public/NNXRuntimeCPUUtils.h
+++ b/Engine/Plugins/Experimental/NNX/Source/NNXRuntimeCPU/Public/NNXRuntimeCPUUtils.h
@@ -138,7 +138,7 @@ inline void BindTensorsToORT(
 		SizesInt64t = MakeUnique<int64_t[]>(CurrentDescriptor.Shape.Num());
 		for (int32 DimIndex = 0; DimIndex < CurrentDescriptor.Shape.Num(); ++DimIndex)
 		{
-			SizesInt64t.Get()[DimIndex] = CurrentDescriptor.Shape[DimIndex];
+			SizesInt64t.Get()[DimIndex] = CurrentDescriptor.Shape.Data[DimIndex];
 		}
 
 		const uint64 ByteCount { InTensorsDescriptors[Index].DataSize };
diff --git a/Engine/Plugins/Experimental/NNX/Source/NNXRuntimeORT/Private/NNXRuntimeORT.cpp b/Engine/Plugins/Experimental/NNX/Source/NNXRuntimeORT/Private/NNXRuntimeORT.cpp
--- a/Engine/Plugins/Experimental/NNX/Source/NNXRuntimeORT/Private/NNXRuntimeORT.cpp
+++ b/Engine/Plugins/Experimental/NNX/Source/NNXRuntimeORT/Private/NNXRuntimeORT.cpp
@@ -262,7 +262,7 @@ bool FMLInferenceModelORT::ConfigureTensors(const bool InIsInput)
 		std::pair<EMLTensorDataType, uint64> TypeAndSize = TranslateTensorTypeORTToNNI(ONNXTensorElementDataTypeEnum);
 		CurrentTensorDescriptor.DataType = TypeAndSize.first;
 		
-		CurrentTensorDescriptor.Shape.SetNumUninitialized(CurrentTensorInfo.GetShape().size());
+		CurrentTensorDescriptor.Shape.Data.SetNumUninitialized(CurrentTensorInfo.GetShape().size());
 		uint8 Index = 0;
 		uint64 TensorNumberElements = 1;
 		for (const int64_t CurrentTensorSize : CurrentTensorInfo.GetShape())
@@ -270,14 +270,14 @@ bool FMLInferenceModelORT::ConfigureTensors(const bool InIsInput)
 			// Negative (variable) dimensions not implemented yet
 			if (CurrentTensorSize < 0)
 			{
-				CurrentTensorDescriptor.Shape[Index] = 1;
+				CurrentTensorDescriptor.Shape.Data[Index] = 1;
 				UE_LOG(LogNNX, Display,
 					TEXT("Negative (i.e., variable) dimensions not allowed yet, hard-coded to 1. Let us know if you really need variable dimensions."
 						" Keep in mind that fixed sizes might allow additional optimizations and speedup of the network during Run()."));
 			}
 			else
 			{
-				CurrentTensorDescriptor.Shape[Index] = CurrentTensorSize;
+				CurrentTensorDescriptor.Shape.Data[Index] = CurrentTensorSize;
 				TensorNumberElements *= CurrentTensorSize;
 			}
 			++Index;
diff --git a/Engine/Plugins/Experimental/NNX/Source/NNXRuntimeORT/Private/NNXRuntimeORTUtils.h b/Engine/Plugins/Experimental/NNX/Source/NNXRuntimeORT/Private/NNXRuntimeORTUtils.h
--- a/Engine/Plugins/Experimental/NNX/Source/NNXRuntimeORT/Private/NNXRuntimeORTUtils.h
+++ b/Engine/Plugins/Experimental/NNX/Source/NNXRuntimeORT/Private/NNXRuntimeORTUtils.h
@@ -138,7 +138,7 @@ inline void BindTensorsToORT(
 		SizesInt64t = MakeUnique<int64_t[]>(CurrentDescriptor.Shape.Num());
 		for (int32 DimIndex = 0; DimIndex < CurrentDescriptor.Shape.Num(); ++DimIndex)
 		{
-			SizesInt64t.Get()[DimIndex] = CurrentDescriptor.Shape[DimIndex];
+			SizesInt64t.Get()[DimIndex] = CurrentDescriptor.Shape.Data[DimIndex];
 		}
 
 		const uint64 ByteCount { InTensorsDescriptors[Index].DataSize };
diff --git a/Engine/Plugins/Experimental/NNX/Source/NNXRuntimeRDG/Private/Hlsl/NNIRuntimeRDGConv.cpp b/Engine/Plugins/Experimental/NNX/Source/NNXRuntimeRDG/Private/Hlsl/NNIRuntimeRDGConv.cpp
--- a/Engine/Plugins/Experimental/NNX/Source/NNXRuntimeRDG/Private/Hlsl/NNIRuntimeRDGConv.cpp
+++ b/Engine/Plugins/Experimental/NNX/Source/NNXRuntimeRDG/Private/Hlsl/NNIRuntimeRDGConv.cpp
@@ -78,11 +78,11 @@ namespace UE::NNIRuntimeRDG::Private::Hlsl
 			constexpr EConvAlgorithm Algorithm = EConvAlgorithm::SharedMemory;
 			constexpr EConvGroupSize GroupSize = EConvGroupSize::Size256;
 
-			TArray<int32> OutputShape = FConvCS::GetOutputShape(Input.Shape, Weights.Shape, AutoPad, Dilations, Strides, Pads);
+			TArray<int32> OutputShape = FConvCS::GetOutputShape(Input.Shape.Data, Weights.Shape.Data, AutoPad, Dilations, Strides, Pads);
 
 			// Set parameters
 			FConvCS::FParameters* Params = GraphBuilder.AllocParameters<FConvCS::FParameters>();
-			FConvCS::FillInParameters(GroupSize, Input.Shape, Weights.Shape, HasBias, AutoPad, Group, Dilations,Strides, Pads, *Params);
+			FConvCS::FillInParameters(GroupSize, Input.Shape.Data, Weights.Shape.Data, HasBias, AutoPad, Group, Dilations,Strides, Pads, *Params);
 			Params->X = GraphBuilder.CreateSRV(FRDGBufferSRVDesc(InInputBindings[0].Buffer, PF_R32_FLOAT));
 			Params->W = GraphBuilder.CreateSRV(FRDGBufferSRVDesc(InInputBindings[1].Buffer, PF_R32_FLOAT));
 			if (InInputBindings.Num() == 3) {
@@ -95,7 +95,7 @@ namespace UE::NNIRuntimeRDG::Private::Hlsl
 			PermutationVector.Set<FConvCS::FConvAlgorithm>(Algorithm);
 			PermutationVector.Set<FConvCS::FConvGroupSize>(GroupSize);
 			PermutationVector.Set<FConvCS::FConvNumDimensions>(NumDimensions);
-			PermutationVector.Set<FConvCS::FConvNumReadsPerThread>(FConvCS::GetNumReadsPerThread(GroupSize, Weights.Shape, Dilations, Strides));
+			PermutationVector.Set<FConvCS::FConvNumReadsPerThread>(FConvCS::GetNumReadsPerThread(GroupSize, Weights.Shape.Data, Dilations, Strides));
 			PermutationVector.Set<FConvCS::FConvHasB>(HasBias);
 			TShaderMapRef<FConvCS> ComputeShader(GetGlobalShaderMap(GMaxRHIFeatureLevel), PermutationVector);
 
diff --git a/Engine/Plugins/Experimental/NNX/Source/NNXRuntimeRDG/Private/Hlsl/NNIRuntimeRDGConvTranspose.cpp b/Engine/Plugins/Experimental/NNX/Source/NNXRuntimeRDG/Private/Hlsl/NNIRuntimeRDGConvTranspose.cpp
--- a/Engine/Plugins/Experimental/NNX/Source/NNXRuntimeRDG/Private/Hlsl/NNIRuntimeRDGConvTranspose.cpp
+++ b/Engine/Plugins/Experimental/NNX/Source/NNXRuntimeRDG/Private/Hlsl/NNIRuntimeRDGConvTranspose.cpp
@@ -80,11 +80,11 @@ namespace UE::NNIRuntimeRDG::Private::Hlsl
 			constexpr EConvTransposeAlgorithm Algorithm = EConvTransposeAlgorithm::SharedMemory;
 			constexpr EConvTransposeGroupSize GroupSize = EConvTransposeGroupSize::Size256;
 
-			TArray<int32> OutputShape = FConvTransposeCS::GetOutputShape(Input.Shape, Weights.Shape, AutoPad, Dilations, Strides, Pads, OutputPadding, Group);
+			TArray<int32> OutputShape = FConvTransposeCS::GetOutputShape(Input.Shape.Data, Weights.Shape.Data, AutoPad, Dilations, Strides, Pads, OutputPadding, Group);
 
 			// Set parameters
 			FConvTransposeCS::FParameters* Params = GraphBuilder.AllocParameters<FConvTransposeCS::FParameters>();
-			FConvTransposeCS::FillInParameters(GroupSize, Input.Shape, Weights.Shape, HasBias, AutoPad, Group, Dilations,Strides, Pads, OutputPadding, *Params);
+			FConvTransposeCS::FillInParameters(GroupSize, Input.Shape.Data, Weights.Shape.Data, HasBias, AutoPad, Group, Dilations,Strides, Pads, OutputPadding, *Params);
 			Params->X = GraphBuilder.CreateSRV(FRDGBufferSRVDesc(InInputBindings[0].Buffer, PF_R32_FLOAT));
 			Params->W = GraphBuilder.CreateSRV(FRDGBufferSRVDesc(InInputBindings[1].Buffer, PF_R32_FLOAT));
 			if (InInputBindings.Num() == 3) {
@@ -97,7 +97,7 @@ namespace UE::NNIRuntimeRDG::Private::Hlsl
 			PermutationVector.Set<FConvTransposeCS::FConvTransposeAlgorithm>(Algorithm);
 			PermutationVector.Set<FConvTransposeCS::FConvTransposeGroupSize>(GroupSize);
 			PermutationVector.Set<FConvTransposeCS::FConvTransposeNumStackDimensions>(NumDimensions);
-			PermutationVector.Set<FConvTransposeCS::FConvTransposeNumReadsPerThread>(FConvTransposeCS::GetNumReadsPerThread(GroupSize, Weights.Shape, Dilations, Strides));
+			PermutationVector.Set<FConvTransposeCS::FConvTransposeNumReadsPerThread>(FConvTransposeCS::GetNumReadsPerThread(GroupSize, Weights.Shape.Data, Dilations, Strides));
 			PermutationVector.Set<FConvTransposeCS::FConvTransposeHasB>(HasBias);
 			TShaderMapRef<FConvTransposeCS> ComputeShader(GetGlobalShaderMap(GMaxRHIFeatureLevel), PermutationVector);
 
diff --git a/Engine/Plugins/Experimental/NNX/Source/NNXRuntimeRDG/Private/Hlsl/NNIRuntimeRDGElementWiseUnary.cpp b/Engine/Plugins/Experimental/NNX/Source/NNXRuntimeRDG/Private/Hlsl/NNIRuntimeRDGElementWiseUnary.cpp
--- a/Engine/Plugins/Experimental/NNX/Source/NNXRuntimeRDG/Private/Hlsl/NNIRuntimeRDGElementWiseUnary.cpp
+++ b/Engine/Plugins/Experimental/NNX/Source/NNXRuntimeRDG/Private/Hlsl/NNIRuntimeRDGElementWiseUnary.cpp
@@ -49,6 +49,30 @@ namespace UE::NNIRuntimeRDG::Private::Hlsl
 
 		virtual void Dispatch(FRDGBuilder& GraphBuilder, TArrayView<const NNX::FMLTensorBinding> InInputBindings, TArrayView<const NNX::FMLTensorBinding> OutOutputBindings) override
 		{
+            //struct RDGTensor
+            //{
+            //    RDGBuffer*
+            //    TensorShape ? Do we need the type? aka TensorDesc?
+            //}
+
+			//Dispatch(FRDGBuilder & GraphBuilder, 
+			         //TConstArrayView<const NNX::RDGTensor> InputBuffers,
+					 //TConstArrayView<const NNX::RDGTensor> OutputBuffers) override
+
+            //or 
+            //Dispatch(FRDGBuilder & GraphBuilder, 
+                     //TConstArrayView<const NNX::FConcreteShape> InputShapes,
+			         //TConstArrayView<const NNX::RDGTensor*> InputBuffers,
+                     //TConstArrayView<const NNX::FConcreteShape> OutputShapes,
+					 //TConstArrayView<const NNX::RDGTensor*> OutputBuffers) override
+
+            //or 
+            //Dispatch(FRDGBuilder & GraphBuilder, 
+                     //TConstArrayView<const NNX::RDGTensor*> Buffers,
+                     //TConstArrayView<const NNX::FConcreteShape> Shapes,
+			         //TConstArrayView<const uint32> InputIndices,
+                     //TConstArrayView<const uint32> OutputIndices) override
+
 			// HACK: This only works for single layer networks
 			FRDGBufferSRVRef InputSRV = GraphBuilder.CreateSRV(FRDGBufferSRVDesc(InInputBindings[0].Buffer, PF_R32_FLOAT));
 			FRDGBufferUAVRef OutputUAV = GraphBuilder.CreateUAV(FRDGBufferUAVDesc(OutOutputBindings[0].Buffer, PF_R32_FLOAT));
diff --git a/Engine/Plugins/Experimental/NNX/Source/NNXRuntimeRDG/Private/Hlsl/NNIRuntimeRDGGemm.cpp b/Engine/Plugins/Experimental/NNX/Source/NNXRuntimeRDG/Private/Hlsl/NNIRuntimeRDGGemm.cpp
--- a/Engine/Plugins/Experimental/NNX/Source/NNXRuntimeRDG/Private/Hlsl/NNIRuntimeRDGGemm.cpp
+++ b/Engine/Plugins/Experimental/NNX/Source/NNXRuntimeRDG/Private/Hlsl/NNIRuntimeRDGGemm.cpp
@@ -54,7 +54,7 @@ namespace UE::NNIRuntimeRDG::Private::Hlsl
 			check(InputB.Shape.Num() == 2);
 			check(InputC.Shape.Num() < 3);
 
-			check(InputC.Shape.Num() != 1 || InputC.Shape[0] != 1); // TODO scalar version not supported yet
+			check(InputC.Shape.Num() != 1 || InputC.Shape.Data[0] != 1); // TODO scalar version not supported yet
 
 			// C is treated as a scalar if there is no valid C, either width or height is zero or C dimension is 1x1
 			bIsCScalar = false; // InputTensors.Num() != 3 || InputC.Sizes[0] * InputC.Sizes[1] < 2;
@@ -108,7 +108,7 @@ namespace UE::NNIRuntimeRDG::Private::Hlsl
 		}
 	};
 
-	bool ValidateGemmOperator(const FMLAttributeMap& AttributeMap, TConstArrayView<EMLTensorDataType> InputTypes)
+	bool ValidateGemmOperator(const FMLAttributeMap& AttributeMap, TConstArrayView<EMLTensorDataType> InputTypes, TConstArrayView<const NNX::FSymbolicTensorShape> InputShapes)
 	{
 		bool bIsValid = true;
 
diff --git a/Engine/Plugins/Experimental/NNX/Source/NNXRuntimeRDG/Private/NNXRuntimeDML.cpp b/Engine/Plugins/Experimental/NNX/Source/NNXRuntimeRDG/Private/NNXRuntimeDML.cpp
--- a/Engine/Plugins/Experimental/NNX/Source/NNXRuntimeRDG/Private/NNXRuntimeDML.cpp
+++ b/Engine/Plugins/Experimental/NNX/Source/NNXRuntimeRDG/Private/NNXRuntimeDML.cpp
@@ -174,7 +174,7 @@ namespace DmlUtil
 		for (int32 i = InputDesc.Shape.Num() - 1; i >= 0; --i)
 		{
 			TensorDesc.Strides[i] = CurrStride;
-			CurrStride *= InputDesc.Shape[i];
+			CurrStride *= InputDesc.Shape.Data[i];
 		}
 	}
 
@@ -192,19 +192,19 @@ namespace DmlUtil
 		
 		for (int32 i = 0; i < (int32) TargetDimension; ++i)
 		{
-			TensorDesc.Sizes[i] = i < DimensionOffset ? 1 : InputDesc.Shape[i - DimensionOffset];
+			TensorDesc.Sizes[i] = i < DimensionOffset ? 1 : InputDesc.Shape.Data[i - DimensionOffset];
 		}
 
 		uint32 CurrStride = 1;
 
 		for (int32 i = TargetDimension - 1; i >= 0; --i)
 		{
-			const bool bBroadcast = TensorDesc.Sizes[i] < TargetDesc.Shape[i];
+			const bool bBroadcast = TensorDesc.Sizes[i] < TargetDesc.Shape.Data[i];
 
 			TensorDesc.Strides[i] = bBroadcast ? 0 : CurrStride;
 			CurrStride *= TensorDesc.Sizes[i];
 
-			TensorDesc.Sizes[i] = TargetDesc.Shape[i];
+			TensorDesc.Sizes[i] = TargetDesc.Shape.Data[i];
 		}
 	}
 
@@ -217,7 +217,7 @@ namespace DmlUtil
 		
 		for (int32 Idx = 0; Idx < Left.Shape.Num(); ++Idx)
 		{
-			if (Left.Shape[Idx] != Right.Shape[Idx])
+			if (Left.Shape.Data[Idx] != Right.Shape.Data[Idx])
 			{
 				return false;
 			}
@@ -343,7 +343,7 @@ protected:
 			return false;
 		}
 
-		DmlTensorDesc.Sizes = TensorDesc.Shape;
+		DmlTensorDesc.Sizes = TensorDesc.Shape.Data;
 		// TODO: Support tensor padding using strides defined in FMLTensorDesc
 		//DmlUtil::SetTensorStrides(DmlTensorDesc, TensorDesc.Strides);
 		
@@ -377,7 +377,7 @@ protected:
 
 		if (DmlUtil::IsSameShape(TensorDesc, BroadcastDesc))
 		{
-			DmlTensorDesc.Sizes = TensorDesc.Shape;
+			DmlTensorDesc.Sizes = TensorDesc.Shape.Data;
 			DmlUtil::SetTensorStrides(DmlTensorDesc, TensorDesc);
 		}
 		else if (TensorDesc.Shape.Num() > BroadcastDesc.Shape.Num())
@@ -1277,20 +1277,25 @@ bool FMLInferenceModelDml::Init(UMLInferenceModel* InModel, FDeviceContextDml* I
 	for (int32 Idx = 0; Idx < Format.Operators.Num(); ++Idx)
 	{
 		const FString TypeName = Format.Operators[Idx].TypeName;
-
-		// HACK: This works only for single layer networks
-		TArray<FMLTensorDesc> OpInputTensors = InputTensors;
-		TArray<FMLTensorDesc> OpOutputTensors = OutputTensors;
-
-		// Attributes
-		FMLAttributeMap Attributes;
 		
+		TArray<FMLTensorDesc> OpInputTensors;
+		TArray<FMLTensorDesc> OpOutputTensors;
+		FMLAttributeMap AttributeMap;
+
+		for (int32 InputTensorIndex : Format.Operators[Idx].InTensors)
+		{
+			OpInputTensors.Emplace(AllTensors[InputTensorIndex]);
+		}
+		for (int32 OutputTensorIndex : Format.Operators[Idx].OutTensors)
+		{
+			OpInputTensors.Emplace(AllTensors[OutputTensorIndex]);
+		}
 		for (const FMLFormatAttributeDesc& Desc : Format.Operators[Idx].Attributes)
 		{
-			Attributes.SetAttribute(Desc.Name, Desc.Value);
+			AttributeMap.SetAttribute(Desc.Name, Desc.Value);
 		}
 
-		FMLOperatorDml* Op = OpCreate(TypeName, OpInputTensors, OpOutputTensors, Attributes);
+		FMLOperatorDml* Op = OpCreate(TypeName, OpInputTensors, OpOutputTensors, AttributeMap);
 
 		if (!Op)
 		{
diff --git a/Engine/Plugins/Experimental/NNX/Source/NNXRuntimeRDG/Private/NNXRuntimeHLSL.cpp b/Engine/Plugins/Experimental/NNX/Source/NNXRuntimeRDG/Private/NNXRuntimeHLSL.cpp
--- a/Engine/Plugins/Experimental/NNX/Source/NNXRuntimeRDG/Private/NNXRuntimeHLSL.cpp
+++ b/Engine/Plugins/Experimental/NNX/Source/NNXRuntimeRDG/Private/NNXRuntimeHLSL.cpp
@@ -71,23 +71,23 @@ bool FMLInferenceModelHlsl::Init(UMLInferenceModel* InModel)
 		return false;
 	}
 
-	if (Format.Operators.Num() > 1)
-	{
-		UE_LOG(LogNNX, Warning, TEXT("Failed to create inference model, currently on single layer models are supported"));
-		return false;
-	}
-
 	// Loop over all operators in the model and create them
 	for (int32 Idx = 0; Idx < Format.Operators.Num(); ++Idx)
 	{
 		const FString TypeName = Format.Operators[Idx].TypeName;
 
-		// HACK: This works only for single layer networks
-		TArray<FMLTensorDesc> OpInputTensors = InputTensors;
-		TArray<FMLTensorDesc> OpOutputTensors = OutputTensors;
-
-		// Attributes
+		TArray<FMLTensorDesc> OpInputTensors;
+		TArray<FMLTensorDesc> OpOutputTensors;
 		FMLAttributeMap AttributeMap;
+
+		for (int32 InputTensorIndex : Format.Operators[Idx].InTensors)
+		{
+			OpInputTensors.Emplace(AllTensors[InputTensorIndex]);
+		}
+		for (int32 OutputTensorIndex : Format.Operators[Idx].OutTensors)
+		{
+			OpInputTensors.Emplace(AllTensors[OutputTensorIndex]);
+		}
 		for (const FMLFormatAttributeDesc& Desc : Format.Operators[Idx].Attributes)
 		{
 			AttributeMap.SetAttribute(Desc.Name, Desc.Value);
diff --git a/Engine/Plugins/Experimental/NNX/Source/NNXRuntimeRDG/Private/NNXRuntimeHLSLHelper.cpp b/Engine/Plugins/Experimental/NNX/Source/NNXRuntimeRDG/Private/NNXRuntimeHLSLHelper.cpp
--- a/Engine/Plugins/Experimental/NNX/Source/NNXRuntimeRDG/Private/NNXRuntimeHLSLHelper.cpp
+++ b/Engine/Plugins/Experimental/NNX/Source/NNXRuntimeRDG/Private/NNXRuntimeHLSLHelper.cpp
@@ -28,7 +28,7 @@ namespace NNX
 			}
 			else
 			{
-				OutShaderParam[i][Idx] = OutShaderParam[i + 1][Idx] * TensorDesc.Shape[i + 1 - Offset];
+				OutShaderParam[i][Idx] = OutShaderParam[i + 1][Idx] * TensorDesc.Shape.Data[i + 1 - Offset];
 			}
 		}
 	}
@@ -41,7 +41,7 @@ namespace NNX
 		for (int32 i = Offset; i < OutputNumdimension; ++i)
 		{
 			// the stride for broadcast dimension is kept as 0
-			if (TensorDesc.Shape[i - Offset] == 1)
+			if (TensorDesc.Shape.Data[i - Offset] == 1)
 			{
 				OutShaderParam[i][Idx] = 0;
 			}
diff --git a/Engine/Plugins/Experimental/NNX/Source/NNXRuntimeRDG/Private/NNXRuntimeRDG.cpp b/Engine/Plugins/Experimental/NNX/Source/NNXRuntimeRDG/Private/NNXRuntimeRDG.cpp
--- a/Engine/Plugins/Experimental/NNX/Source/NNXRuntimeRDG/Private/NNXRuntimeRDG.cpp
+++ b/Engine/Plugins/Experimental/NNX/Source/NNXRuntimeRDG/Private/NNXRuntimeRDG.cpp
@@ -15,7 +15,10 @@ namespace NNX
 //
 //
 //
-bool AlwaysValidValidationFunction(const FMLAttributeMap& AttributeMap, TConstArrayView<EMLTensorDataType> InputTensorTypes)
+bool AlwaysValidValidationFunction(
+	const FMLAttributeMap& AttributeMap, 
+	TConstArrayView<EMLTensorDataType> InputTensorTypes,
+	TConstArrayView<const FSymbolicTensorShape> InputShapes)
 {
 	return true;
 }
@@ -193,32 +196,33 @@ bool FMLInferenceModelRDG::LoadModel(const FNNIModelRaw& InModel, FMLRuntimeForm
 	// Add tensors
 	for (int32 Idx = 0; Idx < Format.Tensors.Num(); ++Idx)
 	{
-		const FMLFormatTensorDesc& TensorDesc = Format.Tensors[Idx];
+		const FMLFormatTensorDesc& FormatTensorDesc = Format.Tensors[Idx];
 
-		FMLTensorDesc	Tensor;
-
-		Tensor.Name = TensorDesc.Name;
-
-		if (TensorDesc.Shape.Num() > FMLTensorDesc::MaxTensorDimension)
-		{
-			// FIXME: Should we report error here?
-		}
-
-		Tensor.Shape = TensorDesc.Shape;
-
-		// NOTE: Set first the DataType prior setting DataSize
-		Tensor.DataType = TensorDesc.DataType;
+		// When handling dynamic input shape FMLTensorDesc should then contain a FSymbolicTensorShape
+        // while actual inference work on FConcreteTensorShape resolved by shape inference.
+		FSymbolicTensorShape SymbolicShape = FSymbolicTensorShape::Make(FormatTensorDesc.Shape);
+		check(SymbolicShape.IsConcrete());
+		FConcreteTensorShape ConcreteShape = FConcreteTensorShape::Make(SymbolicShape);
+		
+		FMLTensorDesc Tensor = FMLTensorDesc::Make(FormatTensorDesc.Name, ConcreteShape, FormatTensorDesc.DataType);
+		
 		Tensor.DataSize = Tensor.GetElemByteSize() * Tensor.Volume();
 
-		if (TensorDesc.Type == EMLFormatTensorType::Input)
+		if (FormatTensorDesc.Type == EMLFormatTensorType::Input)
 		{
 			InputTensors.Add(Tensor);
+			AllTensors.Add(Tensor);
 		}
-		else if (Format.Tensors[Idx].Type == EMLFormatTensorType::Output)
+		else if (FormatTensorDesc.Type == EMLFormatTensorType::Output)
 		{
 			OutputTensors.Add(Tensor);
+			AllTensors.Add(Tensor);
+		}
+		else if (FormatTensorDesc.Type == EMLFormatTensorType::Intermediate)
+		{
+			AllTensors.Add(Tensor);
 		}
-		// FIXME: Read intermediate tensors
+		checkf(FormatTensorDesc.Type != EMLFormatTensorType::None, TEXT("Unsupported tensor type None"));
 	}
 
 	return true;
diff --git a/Engine/Plugins/Experimental/NNX/Source/NNXRuntimeRDG/Private/NNXRuntimeRDG.h b/Engine/Plugins/Experimental/NNX/Source/NNXRuntimeRDG/Private/NNXRuntimeRDG.h
--- a/Engine/Plugins/Experimental/NNX/Source/NNXRuntimeRDG/Private/NNXRuntimeRDG.h
+++ b/Engine/Plugins/Experimental/NNX/Source/NNXRuntimeRDG/Private/NNXRuntimeRDG.h
@@ -80,9 +80,9 @@ protected:
 	bool			bUseManualTransitions;
 };
 
-//TODO jira 167584 remove default validation and declare contract in all HLSL operator (see HLSL Gemm for current example)
-//TODO jira remove default validation and declare contract in all DML operator (see HLSL Gemm for current example)
-bool AlwaysValidValidationFunction(const FMLAttributeMap& AttributeMap, TConstArrayView<EMLTensorDataType> InputTypes);
+//TODO jira 167585 remove default validation and declare contract in all HLSL operator (see HLSL Gemm for current example)
+//TODO jira 167584 remove default validation and declare contract in all DML operator (see HLSL Gemm for current example)
+bool AlwaysValidValidationFunction(const FMLAttributeMap& AttributeMap, TConstArrayView<EMLTensorDataType> InputTypes, TConstArrayView<const FSymbolicTensorShape> InputShapes);
 
 class FInputValidator
 {
@@ -137,7 +137,7 @@ class TOperatorRegistryRDG
 public:
 
 	typedef TOperatorType* (*OperatorCreateFunc)();
-	typedef bool (*OperatorValidateFunc)(const FMLAttributeMap& AttributeMap, TConstArrayView<EMLTensorDataType> InputTypes);
+	typedef bool (*OperatorValidateFunc)(const FMLAttributeMap& AttributeMap, TConstArrayView<EMLTensorDataType> InputTypes, TConstArrayView<const FSymbolicTensorShape> InputShapes);
 
 	static TOperatorRegistryRDG* Get()
 	{
@@ -203,7 +203,7 @@ public:
 		return TEXT("RDG Model validator");
 	}
 	
-	virtual bool ValidateModel(const FNNIModelRaw& InputModel) const override
+	virtual bool ValidateModel(const FNNIModelRaw& InputModel, const FOptimizerOptionsMap& Options) const override
 	{
 		FMLRuntimeFormat	Format;
 
@@ -220,22 +220,17 @@ public:
 		TOperatorRegistryRDG<TOperatorType>* Registry = TOperatorRegistryRDG<TOperatorType>::Get();
 		check(Registry != nullptr);
 
-		// HACK: This works only for single layer networks
-		// we should get the output tensor type from intermediate tensor 
-		// (can we get it from model definition?, otherwise will need
-		// to infer from input shape type + forward propagation and op callbacks)
-		TArray<EMLTensorDataType> InputTensorTypes;
-		for (int32 Idx = 0; Idx < Format.Tensors.Num(); ++Idx)
-		{
-			if (Format.Tensors[Idx].Type == EMLFormatTensorType::Input)
-			{
-				InputTensorTypes.Add(Format.Tensors[Idx].DataType);
-			}
-		}
-
 		for (int32 Idx = 0; Idx < Format.Operators.Num(); ++Idx)
 		{
+			TArray<EMLTensorDataType> InputTensorTypes;
+			TArray<FSymbolicTensorShape> InputTensorShapes;
 			FMLAttributeMap AttributeMap;
+			
+			for (int32 InputTensorIndex: Format.Operators[Idx].InTensors)
+			{
+				InputTensorTypes.Add(Format.Tensors[InputTensorIndex].DataType);
+				InputTensorShapes.Add(FSymbolicTensorShape::Make(Format.Tensors[InputTensorIndex].Shape));
+			}
 			for (const FMLFormatAttributeDesc& Desc : Format.Operators[Idx].Attributes)
 			{
 				AttributeMap.SetAttribute(Desc.Name, Desc.Value);
@@ -249,7 +244,7 @@ public:
 			//will be needed during the execution of the operator.
 			typename TOperatorRegistryRDG<TOperatorType>::OperatorValidateFunc ValidationFn = Registry->OpFindValidation(OpType);
 
-			if (!ValidationFn(AttributeMap, InputTensorTypes))
+			if (!ValidationFn(AttributeMap, InputTensorTypes, InputTensorShapes))
 			{
 				UE_LOG(LogNNX, Warning, TEXT("Hlsl MLOperatorRegistry failed to validate operator:%s"), *OpType);
 				return false;
diff --git a/Engine/Plugins/Experimental/NNX/Source/NNXUtils/Private/NNXModelBuilderNNX.cpp b/Engine/Plugins/Experimental/NNX/Source/NNXUtils/Private/NNXModelBuilderNNX.cpp
--- a/Engine/Plugins/Experimental/NNX/Source/NNXUtils/Private/NNXModelBuilderNNX.cpp
+++ b/Engine/Plugins/Experimental/NNX/Source/NNXUtils/Private/NNXModelBuilderNNX.cpp
@@ -120,12 +120,18 @@ public:
 	{
 		int Idx = NNXTensorCast(Tensor);
 		
-		if (Idx < 0)
+		if (Idx < 0 || Idx >= Format.Tensors.Num())
 		{
 			UE_LOG(LogNNX, Warning, TEXT("Failed to add input tensor, invalid tensor index"));
 			return false;
 		}
 		
+		if (Format.Tensors[Idx].Type != EMLFormatTensorType::None)
+		{
+			UE_LOG(LogNNX, Warning, TEXT("Failed to add input tensor, tensor usage already setupped"));
+			return false;
+		}
+		
 		Format.Tensors[Idx].Type = EMLFormatTensorType::Input;
 
 		return true;
@@ -136,12 +142,18 @@ public:
 	{
 		int Idx = NNXTensorCast(Tensor);
 
-		if (Idx < 0)
+		if (Idx < 0 || Idx >= Format.Tensors.Num())
 		{
 			UE_LOG(LogNNX, Warning, TEXT("Failed to add output tensor, invalid tensor index"));
 			return false;
 		}
 
+		if (Format.Tensors[Idx].Type != EMLFormatTensorType::None)
+		{
+			UE_LOG(LogNNX, Warning, TEXT("Failed to add output tensor, tensor usage already setupped"));
+			return false;
+		}
+
 		Format.Tensors[Idx].Type = EMLFormatTensorType::Output;
 
 		return true;
@@ -166,7 +178,11 @@ public:
 		int OpIdx = NNXOperatorCast(Op);
 		int TensorIdx = NNXTensorCast(Tensor);
 		
-		// TODO: Set tensor type
+		if (TensorIdx < 0 || TensorIdx >= Format.Tensors.Num())
+		{
+			UE_LOG(LogNNX, Warning, TEXT("Failed to add operator input tensor, invalid tensor index"));
+			return false;
+		}
 
 		Format.Operators[OpIdx].InTensors.Emplace(TensorIdx);
 
@@ -178,9 +194,24 @@ public:
 	{
 		int OpIdx = NNXOperatorCast(Op);
 		int TensorIdx = NNXTensorCast(Tensor);
+		
+		if (TensorIdx < 0 || TensorIdx >= Format.Tensors.Num())
+		{
+			UE_LOG(LogNNX, Warning, TEXT("Failed to add operator output tensor, invalid tensor index"));
+			return false;
+		}
 
-		// TODO: Set tensor type
+		if (Format.Tensors[TensorIdx].Type == EMLFormatTensorType::Input)
+		{
+			UE_LOG(LogNNX, Warning, TEXT("Failed to add output tensor, tensor usage already setupped to input"));
+			return false;
+		}
 
+		if (Format.Tensors[TensorIdx].Type == EMLFormatTensorType::None)
+		{
+			Format.Tensors[TensorIdx].Type = EMLFormatTensorType::Intermediate;
+		}
+		
 		Format.Operators[OpIdx].OutTensors.Emplace(TensorIdx);
 
 		return true;
diff --git a/Engine/Plugins/Experimental/NNX/Source/NNXUtils/Private/NNXModelBuilderONNX.cpp b/Engine/Plugins/Experimental/NNX/Source/NNXUtils/Private/NNXModelBuilderONNX.cpp
--- a/Engine/Plugins/Experimental/NNX/Source/NNXUtils/Private/NNXModelBuilderONNX.cpp
+++ b/Engine/Plugins/Experimental/NNX/Source/NNXUtils/Private/NNXModelBuilderONNX.cpp
@@ -297,7 +297,7 @@ NNXUTILS_API bool CreateONNXModelForOperator(const FString& OperatorName, TConst
 	for (int32 Idx = 0; Idx < InInputTensors.Num(); ++Idx)
 	{
 		const FMLTensorDesc& Desc = InInputTensors[Idx];
-		IMLModelBuilder::HTensor Tensor = Builder->AddTensor(Desc.Name, Desc.DataType, MakeArrayView((const int32*) Desc.Shape.GetData(), Desc.Shape.Num()));
+		IMLModelBuilder::HTensor Tensor = Builder->AddTensor(Desc.Name, Desc.DataType, MakeArrayView((const int32*) Desc.Shape.Data.GetData(), Desc.Shape.Num()));
 
 		InputTensors.Emplace(Tensor);
 		Builder->AddInput(Tensor);
@@ -308,7 +308,7 @@ NNXUTILS_API bool CreateONNXModelForOperator(const FString& OperatorName, TConst
 	for (int32 Idx = 0; Idx < InOutputTensors.Num(); ++Idx)
 	{
 		const FMLTensorDesc& Desc = InOutputTensors[Idx];
-		IMLModelBuilder::HTensor Tensor = Builder->AddTensor(Desc.Name, Desc.DataType, MakeArrayView((const int32*) Desc.Shape.GetData(), Desc.Shape.Num()));
+		IMLModelBuilder::HTensor Tensor = Builder->AddTensor(Desc.Name, Desc.DataType, MakeArrayView((const int32*) Desc.Shape.Data.GetData(), Desc.Shape.Num()));
 
 		OutputTensors.Emplace(Tensor);
 		Builder->AddOutput(Tensor);
diff --git a/Engine/Plugins/Experimental/NNX/Source/NNXUtils/Private/NNXModelOptimizer.cpp b/Engine/Plugins/Experimental/NNX/Source/NNXUtils/Private/NNXModelOptimizer.cpp
--- a/Engine/Plugins/Experimental/NNX/Source/NNXUtils/Private/NNXModelOptimizer.cpp
+++ b/Engine/Plugins/Experimental/NNX/Source/NNXUtils/Private/NNXModelOptimizer.cpp
@@ -303,14 +303,14 @@ public:
 		}
 	}
 
-	bool IsModelValid(const FNNIModelRaw& ModelToValidate)
+	bool IsModelValid(const FNNIModelRaw& ModelToValidate, const FOptimizerOptionsMap& Options)
 	{
 		bool bIsModelValid = true;
 
 		for (TSharedPtr<IModelValidator>& Validator : Validators)
 		{
 			check(Validator.IsValid());
-			if (!Validator->ValidateModel(ModelToValidate))
+			if (!Validator->ValidateModel(ModelToValidate, Options))
 			{
 				UE_LOG(LogNNX, Warning, TEXT("Model validator %s detected an error."), *(Validator->GetName()));
 				bIsModelValid = false;
@@ -321,7 +321,7 @@ public:
 
 	bool ApplyAllPassesAndValidations(FNNIModelRaw& OptimizedModel, const FOptimizerOptionsMap& Options)
 	{
-		if (!IsModelValid(OptimizedModel))
+		if (!IsModelValid(OptimizedModel, Options))
 		{
 			UE_LOG(LogNNX, Warning, TEXT("Model is not valid, skipping optimization passes."));
 			return false;
@@ -336,7 +336,7 @@ public:
 				UE_LOG(LogNNX, Warning, TEXT("Error while executing model optimisation pass %s."), *(Pass->GetName()));
 				return false;
 			}
-			if (!IsModelValid(OptimizedModel))
+			if (!IsModelValid(OptimizedModel, Options))
 			{
 				UE_LOG(LogNNX, Warning, TEXT("Model validation failed after optimisation pass %s."), *(Pass->GetName()));
 				return false;
@@ -368,7 +368,7 @@ public:
 		return TEXT("ONNX Model validator");
 	}
 
-	virtual bool ValidateModel(const FNNIModelRaw& InputModel) const override
+	virtual bool ValidateModel(const FNNIModelRaw& InputModel, const FOptimizerOptionsMap& Options) const override
 	{
 		FMLRuntimeFormat	Format;
 
diff --git a/Engine/Plugins/Experimental/NNX/Source/NNXUtils/Private/NNXUtilsModule.cpp b/Engine/Plugins/Experimental/NNX/Source/NNXUtils/Private/NNXUtilsModule.cpp
--- a/Engine/Plugins/Experimental/NNX/Source/NNXUtils/Private/NNXUtilsModule.cpp
+++ b/Engine/Plugins/Experimental/NNX/Source/NNXUtils/Private/NNXUtilsModule.cpp
@@ -88,8 +88,8 @@ static FAutoConsoleCommand ConsoleCmdNNXOptimizeModel(
 			}
 			else if (!OpName.IsEmpty())
 			{
-				FMLTensorDesc InputTensor = FMLTensorDesc::Make(TEXT("in"), { 1, 512 }, EMLTensorDataType::Float);
-				FMLTensorDesc OutputTensor = FMLTensorDesc::Make(TEXT("out"), { 1, 512 }, EMLTensorDataType::Float);
+				FMLTensorDesc InputTensor = FMLTensorDesc::Make(TEXT("in"), FConcreteTensorShape::Make({1, 512 }), EMLTensorDataType::Float);
+				FMLTensorDesc OutputTensor = FMLTensorDesc::Make(TEXT("out"), FConcreteTensorShape::Make({ 1, 512 }), EMLTensorDataType::Float);
 
 				if (!NNX::CreateONNXModelForOperator(OpName, MakeArrayView(&InputTensor, 1), MakeArrayView(&OutputTensor, 1), ONNXModel))
 				{
@@ -137,8 +137,8 @@ static FAutoConsoleCommand ConsoleCmdNNXCreateModel(
 			if (Args.Num() > 0)
 			{
 				FNNIModelRaw ONNXModel;
-				FMLTensorDesc InputTensor = FMLTensorDesc::Make(TEXT("in"), { 1, 512 }, EMLTensorDataType::Float);
-				FMLTensorDesc OutputTensor = FMLTensorDesc::Make(TEXT("out"), { 1, 512 }, EMLTensorDataType::Float);
+				FMLTensorDesc InputTensor = FMLTensorDesc::Make(TEXT("in"), FConcreteTensorShape::Make({ 1, 512 }), EMLTensorDataType::Float);
+				FMLTensorDesc OutputTensor = FMLTensorDesc::Make(TEXT("out"), FConcreteTensorShape::Make({ 1, 512 }), EMLTensorDataType::Float);
 
 				if (!NNX::CreateONNXModelForOperator(Args[0], MakeArrayView(&InputTensor, 1), MakeArrayView(&OutputTensor, 1), ONNXModel))
 				{
