Prepare to extend operator attribute support and some cleanup

diff --git a/Engine/Plugins/Experimental/NNX/Source/NNXCore/Private/NNXModelOptimizerInterface.h b/Engine/Plugins/Experimental/NNX/Source/NNXCore/Private/NNXModelOptimizerInterface.h
--- a/Engine/Plugins/Experimental/NNX/Source/NNXCore/Private/NNXModelOptimizerInterface.h
+++ b/Engine/Plugins/Experimental/NNX/Source/NNXCore/Private/NNXModelOptimizerInterface.h
@@ -1,14 +1,15 @@
 // Copyright Epic Games, Inc. All Rights Reserved.
 #pragma once
 
-#include "NNXRuntimeFormat.h"
+struct FNNIModelRaw;
+namespace UE::NNECore { class FAttributeMap; }
 
 namespace NNX
 {
 	
-//TODO create a dedicated FOptimizerOptionsMap when function diverge with FMLAttributeMap
+//TODO create a dedicated FOptimizerOptionsMap when function diverge with FAttributeMap
 //example introduction of sparse tensor to FMLAttribute witch make no sense as an optimizer attribute
-using FOptimizerOptionsMap = FMLAttributeMap;
+using FOptimizerOptionsMap = UE::NNECore::FAttributeMap;
 
 /** Interface class for NNX model validator */
 class IModelValidator
diff --git a/Engine/Plugins/Experimental/NNX/Source/NNXCore/Public/NNECoreAttributeDataType.h b/Engine/Plugins/Experimental/NNX/Source/NNXCore/Public/NNECoreAttributeDataType.h
new file mode 100644
--- /dev/null
+++ b/Engine/Plugins/Experimental/NNX/Source/NNXCore/Public/NNECoreAttributeDataType.h
@@ -0,0 +1,37 @@
+// Copyright Epic Games, Inc. All Rights Reserved.
+
+#pragma once
+
+#include "UObject/Class.h"
+#include "UObject/Object.h"
+
+#include "NNECoreAttributeDataType.generated.h"
+
+/**
+ * Attribute data types
+ * 
+ * Note: also extend NNEAttributeValueTraits.h for more type support
+ */
+UENUM()
+enum class ENNEAttributeDataType : uint8
+{
+	None,
+	Float,								//!< 32-bit floating number
+	Int32								//!< 32-bit signed integer
+	//TODO add more AttributeDataType support
+};
+
+/**
+ * @return ENNEAttributeDataType from the string passed in
+ */
+inline void LexFromString(ENNEAttributeDataType& OutValue, const TCHAR* StringVal)
+{
+	int64 EnumVal = StaticEnum<ENNEAttributeDataType>()->GetValueByName(StringVal);
+	if (EnumVal == INDEX_NONE)
+	{
+		OutValue = ENNEAttributeDataType::None;
+		ensureMsgf(false, TEXT("ENNEAttributeDataType LexFromString didn't have a match for '%s'"), StringVal);
+	}
+
+	OutValue = (ENNEAttributeDataType)EnumVal;
+};
\ No newline at end of file
diff --git a/Engine/Plugins/Experimental/NNX/Source/NNXCore/Public/NNECoreAttributeMap.h b/Engine/Plugins/Experimental/NNX/Source/NNXCore/Public/NNECoreAttributeMap.h
new file mode 100644
--- /dev/null
+++ b/Engine/Plugins/Experimental/NNX/Source/NNXCore/Public/NNECoreAttributeMap.h
@@ -0,0 +1,80 @@
+// Copyright Epic Games, Inc. All Rights Reserved.
+
+#pragma once
+
+#include "NNECoreAttributeValue.h"
+
+namespace UE::NNECore
+{
+
+class FAttributeMap
+{
+public:
+	//Set attribute
+	void SetAttribute(const FString& Name, const FNNEAttributeValue& Value)
+	{
+		checkf(nullptr == Attributes.FindByPredicate([Name](const FEntry& Other) { return Other.Name == Name; }), TEXT("Attribute name should be unique"));
+		Attributes.Emplace(Name, Value);
+	}
+
+	template<typename T>
+	T GetValue(const FString& Name) const
+	{
+		const FNNEAttributeValue* Value = GetAttributeValue(Name);
+		checkf(Value != nullptr, TEXT("Required attribute %s not found"), *Name);
+		return Value->GetValue<T>();
+	}
+
+	//Query attributes
+	template<typename T>
+	T GetValueOrDefault(const FString& Name, T Default) const
+	{
+		const FNNEAttributeValue* Value = GetAttributeValue(Name);
+		return Value == nullptr ? Default : Value->GetValue<T>();
+	}
+	
+	const FNNEAttributeValue* GetAttributeValue(const FString& Name) const
+	{
+		const FEntry * entry = Attributes.FindByPredicate([Name](const FEntry& Other) { return Other.Name == Name; });
+		if (entry != nullptr)
+		{
+			return &entry->Value;
+		}
+		return nullptr;
+	}
+
+	//Iteration
+	int32 Num() const
+	{
+		return Attributes.Num();
+	}
+
+	const FString& GetName(int32 Idx) const
+	{
+		check(Idx < Num());
+		return Attributes[Idx].Name;
+	}
+
+	const FNNEAttributeValue& GetAttributeValue(int32 Idx) const
+	{
+		check(Idx < Num());
+		return Attributes[Idx].Value;
+	}
+
+private:
+
+	struct FEntry
+	{
+		FEntry(const FString& InName, const FNNEAttributeValue& InValue)
+			: Name(InName), Value(InValue)
+		{
+		}
+		
+		FString Name;
+		FNNEAttributeValue Value;
+	};
+
+	TArray<FEntry> Attributes;
+};
+
+} // namespace UE::NNECore
\ No newline at end of file
diff --git a/Engine/Plugins/Experimental/NNX/Source/NNXCore/Public/NNECoreAttributeValue.h b/Engine/Plugins/Experimental/NNX/Source/NNXCore/Public/NNECoreAttributeValue.h
new file mode 100644
--- /dev/null
+++ b/Engine/Plugins/Experimental/NNX/Source/NNXCore/Public/NNECoreAttributeValue.h
@@ -0,0 +1,57 @@
+// Copyright Epic Games, Inc. All Rights Reserved.
+
+#pragma once
+
+#include "UObject/Class.h"
+#include "UObject/Object.h"
+#include "Serialization/MemoryWriter.h"
+#include "Serialization/MemoryReader.h"
+#include "NNECoreAttributeValueTraits.h"
+
+#include "NNECoreAttributeValue.generated.h"
+
+USTRUCT()
+struct FNNEAttributeValue
+{
+	GENERATED_USTRUCT_BODY()
+
+	FNNEAttributeValue()
+	{
+		Type = ENNEAttributeDataType::None;
+	}
+
+	template<typename T>
+	explicit FNNEAttributeValue(T InValue)
+	{
+		FMemoryWriter writer(Value, true);
+		writer << InValue;
+
+		Type = TNNEAttributeValueTraits<T>::GetType();
+	}
+
+	template<typename T>
+	T GetValue() const
+	{
+		check((Type == TNNEAttributeValueTraits<T>::GetType()));
+
+		T Result;
+
+		FMemoryReader Reader(Value, true);
+		Reader << Result;
+
+		return Result;
+	}
+
+	ENNEAttributeDataType GetType() const
+	{
+		return Type;
+	}
+	
+private:
+	
+	UPROPERTY(VisibleAnywhere, Category = "Neural Network Inference")
+	ENNEAttributeDataType	Type;
+	
+	UPROPERTY(VisibleAnywhere, Category = "Neural Network Inference")
+	TArray<uint8>			Value;
+};
\ No newline at end of file
diff --git a/Engine/Plugins/Experimental/NNX/Source/NNXCore/Public/NNECoreAttributeValueTraits.h b/Engine/Plugins/Experimental/NNX/Source/NNXCore/Public/NNECoreAttributeValueTraits.h
new file mode 100644
--- /dev/null
+++ b/Engine/Plugins/Experimental/NNX/Source/NNXCore/Public/NNECoreAttributeValueTraits.h
@@ -0,0 +1,25 @@
+// Copyright Epic Games, Inc. All Rights Reserved.
+
+#pragma once
+
+#include "NNECoreAttributeDataType.h"
+
+template<typename T> struct TNNEAttributeValueTraits
+{
+	static constexpr ENNEAttributeDataType GetType()
+	{
+		static_assert(!sizeof(T), "Attribute value trait must be specialized for this type.");
+		return ENNEAttributeDataType::None;
+	}
+};
+
+// Attribute specializations
+template<> struct TNNEAttributeValueTraits<float>
+{
+	static constexpr ENNEAttributeDataType GetType() { return ENNEAttributeDataType::Float; }
+};
+
+template<> struct TNNEAttributeValueTraits<int32>
+{
+	static constexpr ENNEAttributeDataType GetType() { return ENNEAttributeDataType::Int32; }
+};
\ No newline at end of file
diff --git a/Engine/Plugins/Experimental/NNX/Source/NNXCore/Public/NNXRuntimeFormat.h b/Engine/Plugins/Experimental/NNX/Source/NNXCore/Public/NNXRuntimeFormat.h
--- a/Engine/Plugins/Experimental/NNX/Source/NNXCore/Public/NNXRuntimeFormat.h
+++ b/Engine/Plugins/Experimental/NNX/Source/NNXCore/Public/NNXRuntimeFormat.h
@@ -3,6 +3,7 @@
 #pragma once
 
 #include "UObject/Object.h"
+#include "NNECoreAttributeValue.h"
 #include "NNXTypes.h"
 #include "NNXRuntime.h"
 #include "NNXRuntimeFormat.generated.h"
@@ -48,7 +49,7 @@ struct FMLFormatAttributeDesc
 	FString Name;
 
 	UPROPERTY(VisibleAnywhere, Category = "Neural Network Inference")
-	FMLAttributeValue Value;
+	FNNEAttributeValue Value;
 };
 
 USTRUCT()
diff --git a/Engine/Plugins/Experimental/NNX/Source/NNXCore/Public/NNXTypes.h b/Engine/Plugins/Experimental/NNX/Source/NNXCore/Public/NNXTypes.h
--- a/Engine/Plugins/Experimental/NNX/Source/NNXCore/Public/NNXTypes.h
+++ b/Engine/Plugins/Experimental/NNX/Source/NNXCore/Public/NNXTypes.h
@@ -2,8 +2,8 @@
 
 #pragma once
 
-#include "UObject/Object.h"
 #include "UObject/Class.h"
+#include "UObject/Object.h"
 #include "NNXTypes.generated.h"
 
 /** Tensor data types */
@@ -29,163 +29,6 @@ enum class EMLTensorDataType : uint8
 	BFloat16							//!< 16-bit floating number
 };
 
-/** Attribute data types */
-UENUM()
-enum class EMLAttributeDataType : uint8
-{
-	None,
-	Float,								//!< 32-bit floating number
-	Int32								//!< 32-bit signed integer
-    //TODO add more AttributeDataType support
-};
-
-/**
- * @return EMLAttributeDataType from the string passed in
- */
-inline void LexFromString(EMLAttributeDataType& OutValue, const TCHAR* StringVal)
-{
-	int64 EnumVal = StaticEnum<EMLAttributeDataType>()->GetValueByName(StringVal);
-	if (EnumVal == INDEX_NONE)
-	{
-		OutValue = EMLAttributeDataType::None;
-		ensureMsgf(false, TEXT("EMLAttributeDataType LexFromString didn't have a match for '%s'"), StringVal);
-	}
-
-	OutValue = (EMLAttributeDataType)EnumVal;
-};
-
-USTRUCT()
-struct FMLAttributeValue
-{
-	GENERATED_USTRUCT_BODY()
-
-	FMLAttributeValue()
-	{
-		Type = EMLAttributeDataType::None;
-	}
-
-    //TODO add more AttributeDataType support
-	explicit FMLAttributeValue(float Value)
-	{
-		Type = EMLAttributeDataType::Float;
-		Data.SetNumUninitialized(sizeof(float));
-		FMemory::Memcpy(Data.GetData(), &Value, sizeof(float));
-	}
-
-	explicit FMLAttributeValue(int Value)
-	{
-		Type = EMLAttributeDataType::Int32;
-		Data.SetNumUninitialized(sizeof(int32));
-		FMemory::Memcpy(Data.GetData(), &Value, sizeof(int32));
-	}
-
-	float AsFloat() const
-	{
-		check(Type == EMLAttributeDataType::Float);
-		return *(float*)Data.GetData();
-	}
-
-	int32 AsInt32() const
-	{
-		check(Type == EMLAttributeDataType::Int32);
-		return *(int32*)Data.GetData();
-	}
-
-	EMLAttributeDataType GetType() const
-	{
-		return Type;
-	}
-	
-private:
-	
-	UPROPERTY(VisibleAnywhere, Category = "Neural Network Inference")
-	EMLAttributeDataType	Type;
-	
-	UPROPERTY(VisibleAnywhere, Category = "Neural Network Inference")
-	TArray<uint8>			Data;
-};
-
-class FMLAttributeMap
-{
-public:
-	//Set attribute
-	void SetAttribute(const FString& Name, const FMLAttributeValue& Value)
-	{
-		checkf(nullptr == Attributes.FindByPredicate([Name](const FEntry& Other) { return Other.Name == Name; }), TEXT("Attribute name should be unique"));
-		Attributes.Emplace(Name, Value);
-	}
-
-	//Query attributes
-	float GetFloat(const FString& Name) const
-	{
-		const FMLAttributeValue* Value = GetAttributeValue(Name);
-		checkf(Value != nullptr, TEXT("Required attribute %s not found"), *Name);
-		return Value->AsFloat();
-	}
-
-	int32 GetInt32(const FString& Name) const
-	{
-		const FMLAttributeValue* Value = GetAttributeValue(Name);
-		checkf(Value != nullptr, TEXT("Required attribute %s not found"), *Name);
-		return Value->AsInt32();
-	}
-
-	float GetOptionalFloat(const FString& Name, float Default) const
-	{
-		const FMLAttributeValue* Value = GetAttributeValue(Name);
-		return Value == nullptr ? Default : Value->AsFloat();
-	}
-
-	int32 GetOptionalInt32(const FString& Name, int32 Default) const
-	{
-		const FMLAttributeValue* Value = GetAttributeValue(Name);
-		return Value == nullptr ? Default : Value->AsInt32();
-	}
-	
-	const FMLAttributeValue* GetAttributeValue(const FString& Name) const
-	{
-		const FEntry * entry = Attributes.FindByPredicate([Name](const FEntry& Other) { return Other.Name == Name; });
-		if (entry != nullptr)
-		{
-			return &entry->Value;
-		}
-		return nullptr;
-	}
-
-	//Iteration
-	int32 Num() const
-	{
-		return Attributes.Num();
-	}
-
-	const FString& GetName(int32 Idx) const
-	{
-		check(Idx < Num());
-		return Attributes[Idx].Name;
-	}
-
-	const FMLAttributeValue& GetAttributeValue(int32 Idx) const
-	{
-		check(Idx < Num());
-		return Attributes[Idx].Value;
-	}
-
-private:
-
-	struct FEntry
-	{
-		FEntry(const FString& InName, const FMLAttributeValue& InValue)
-			: Name(InName), Value(InValue)
-		{
-		}
-		
-		FString Name;
-		FMLAttributeValue Value;
-	};
-
-	TArray<FEntry> Attributes;
-};
-
 namespace NNX
 {
 
diff --git a/Engine/Plugins/Experimental/NNX/Source/NNXQA/Private/NNXQAJsonUtils.h b/Engine/Plugins/Experimental/NNX/Source/NNXQA/Private/NNXQAJsonUtils.h
--- a/Engine/Plugins/Experimental/NNX/Source/NNXQA/Private/NNXQAJsonUtils.h
+++ b/Engine/Plugins/Experimental/NNX/Source/NNXQA/Private/NNXQAJsonUtils.h
@@ -65,7 +65,7 @@ namespace Json
 	struct FTestAttribute : FJsonSerializable
 	{
 		FString Name;
-		FMLAttributeValue Value;
+		FNNEAttributeValue Value;
 
 		BEGIN_JSON_SERIALIZER
 			JSON_SERIALIZE("name", Name);
@@ -77,28 +77,28 @@ namespace Json
 			FString TypeStr;
 			Serializer.Serialize(TEXT("type"), TypeStr);
 
-			EMLAttributeDataType Type = EMLAttributeDataType::Float;
+			ENNEAttributeDataType Type = ENNEAttributeDataType::Float;
 			if (!TypeStr.IsEmpty()) LexFromString(Type, *TypeStr);
 
 			switch (Type)
 			{
-				case EMLAttributeDataType::Float:
+				case ENNEAttributeDataType::Float:
 				{
 					float TmpValue;
 					JSON_SERIALIZE("value", TmpValue);
-					Value = FMLAttributeValue(TmpValue);
+					Value = FNNEAttributeValue(TmpValue);
 					break;
 				}
 
-				case EMLAttributeDataType::Int32:
+				case ENNEAttributeDataType::Int32:
 				{
 					int TmpValue;
 					JSON_SERIALIZE("value", TmpValue);
-					Value = FMLAttributeValue(TmpValue);
+					Value = FNNEAttributeValue(TmpValue);
 					break;
 				}
 				default:
-					check(Type == EMLAttributeDataType::None);
+					check(Type == ENNEAttributeDataType::None);
 			}
 		END_JSON_SERIALIZER
 	};
diff --git a/Engine/Plugins/Experimental/NNX/Source/NNXQA/Private/NNXQAParametricTest.cpp b/Engine/Plugins/Experimental/NNX/Source/NNXQA/Private/NNXQAParametricTest.cpp
--- a/Engine/Plugins/Experimental/NNX/Source/NNXQA/Private/NNXQAParametricTest.cpp
+++ b/Engine/Plugins/Experimental/NNX/Source/NNXQA/Private/NNXQAParametricTest.cpp
@@ -6,6 +6,7 @@
 #include "NNXCore.h"
 #include "NNXTypes.h"
 #include "NNXRuntimeFormat.h"
+#include "NNECoreAttributeMap.h"
 #include "NNXQAUtils.h"
 #include "NNXQAJsonUtils.h"
 #include "NNXModelBuilder.h"
@@ -205,13 +206,13 @@ namespace Test
 		{
 			if (AttributeMap.Attributes.IsEmpty()) return TEXT("");
 
-			auto AttributeToString = [] (const FMLAttributeValue &Value) -> FString {
-				EMLAttributeDataType Type = Value.GetType();
+			auto AttributeToString = [] (const FNNEAttributeValue &Value) -> FString {
+				ENNEAttributeDataType Type = Value.GetType();
 				
 				switch(Type)
 				{
-					case EMLAttributeDataType::Float: return FString::SanitizeFloat(Value.AsFloat());
-					case EMLAttributeDataType::Int32: return FString::FromInt(Value.AsInt32());
+					case ENNEAttributeDataType::Float: return FString::SanitizeFloat(Value.GetValue<float>());
+					case ENNEAttributeDataType::Int32: return FString::FromInt(Value.GetValue<int32>());
 				}
 				return TEXT("-");
 			};
diff --git a/Engine/Plugins/Experimental/NNX/Source/NNXQA/Private/NNXQAUtils.cpp b/Engine/Plugins/Experimental/NNX/Source/NNXQA/Private/NNXQAUtils.cpp
--- a/Engine/Plugins/Experimental/NNX/Source/NNXQA/Private/NNXQAUtils.cpp
+++ b/Engine/Plugins/Experimental/NNX/Source/NNXQA/Private/NNXQAUtils.cpp
@@ -1,11 +1,12 @@
 // Copyright Epic Games, Inc. All Rights Reserved.
 
 #include "NNXQAUtils.h"
-#include "NNXCore.h"
-#include "NNXInferenceModel.h"
-#include "NNXModelOptimizer.h"
 #include "HAL/UnrealMemory.h"
 #include "Kismet/GameplayStatics.h"
+#include "NNECoreAttributeMap.h"
+#include "NNXCore.h"
+#include "NNXInferenceModel.h"
+#include "NNXModelOptimizerInterface.h"
 
 #include <functional>
 
diff --git a/Engine/Plugins/Experimental/NNX/Source/NNXQA/Private/NNXQAUtils.h b/Engine/Plugins/Experimental/NNX/Source/NNXQA/Private/NNXQAUtils.h
--- a/Engine/Plugins/Experimental/NNX/Source/NNXQA/Private/NNXQAUtils.h
+++ b/Engine/Plugins/Experimental/NNX/Source/NNXQA/Private/NNXQAUtils.h
@@ -4,6 +4,7 @@
 
 #include "Containers/UnrealString.h"
 #include "NNXRuntime.h"
+#include "NNECoreAttributeMap.h"
 
 struct FNNIModelRaw;
 
@@ -35,7 +36,7 @@ namespace Test
 			TMap<FString, float> RelativeErrorPercentForRuntime;
 			TArray<FMLTensorDesc> Inputs;
 			TArray<FMLTensorDesc> Outputs;
-			FMLAttributeMap AttributeMap;
+			UE::NNECore::FAttributeMap AttributeMap;
 			TArray<FString> Tags;
 			TArray<FString> AutomationExcludedRuntime;
 			TArray<FString> AutomationExcludedPlatform;
diff --git a/Engine/Plugins/Experimental/NNX/Source/NNXRuntimeCPU/Public/NNXRuntimeCPU.h b/Engine/Plugins/Experimental/NNX/Source/NNXRuntimeCPU/Public/NNXRuntimeCPU.h
--- a/Engine/Plugins/Experimental/NNX/Source/NNXRuntimeCPU/Public/NNXRuntimeCPU.h
+++ b/Engine/Plugins/Experimental/NNX/Source/NNXRuntimeCPU/Public/NNXRuntimeCPU.h
@@ -7,6 +7,7 @@
 #include "NNXRuntime.h"
 #include "NNXInferenceModel.h"
 #include "NNIProfilingStatistics.h"
+#include "NNXModelOptimizerInterface.h"
 
 #include "NNXThirdPartyWarningDisabler.h"
 NNX_THIRD_PARTY_INCLUDES_START
diff --git a/Engine/Plugins/Experimental/NNX/Source/NNXRuntimeORT/Private/NNXRuntimeORT.h b/Engine/Plugins/Experimental/NNX/Source/NNXRuntimeORT/Private/NNXRuntimeORT.h
--- a/Engine/Plugins/Experimental/NNX/Source/NNXRuntimeORT/Private/NNXRuntimeORT.h
+++ b/Engine/Plugins/Experimental/NNX/Source/NNXRuntimeORT/Private/NNXRuntimeORT.h
@@ -6,6 +6,7 @@
 #include "NNXCore.h"
 #include "NNXRuntime.h"
 #include "NNXInferenceModel.h"
+#include "NNXModelOptimizerInterface.h"
 #include "NNIProfilingStatistics.h"
 #include "NNXThirdPartyWarningDisabler.h"
 
diff --git a/Engine/Plugins/Experimental/NNX/Source/NNXRuntimeRDG/Private/Hlsl/NNIRuntimeRDGConv.cpp b/Engine/Plugins/Experimental/NNX/Source/NNXRuntimeRDG/Private/Hlsl/NNIRuntimeRDGConv.cpp
--- a/Engine/Plugins/Experimental/NNX/Source/NNXRuntimeRDG/Private/Hlsl/NNIRuntimeRDGConv.cpp
+++ b/Engine/Plugins/Experimental/NNX/Source/NNXRuntimeRDG/Private/Hlsl/NNIRuntimeRDGConv.cpp
@@ -2,6 +2,7 @@
 
 #include "NNIRuntimeRDGConv.h"
 #include "NNIHlslShadersConvCS.h"
+#include "NNECoreAttributeMap.h"
 #include "NNXRuntimeHLSLHelper.h"
 
 namespace UE::NNIRuntimeRDG::Private::Hlsl
@@ -45,7 +46,7 @@ namespace UE::NNIRuntimeRDG::Private::Hlsl
 
 	public:
 
-		virtual bool Initialize(TArrayView<const NNX::FMLTensorDesc> InputTensors, TArrayView<const NNX::FMLTensorDesc> OutputTensors, const FMLAttributeMap& Attributes) override
+		virtual bool Initialize(TArrayView<const NNX::FMLTensorDesc> InputTensors, TArrayView<const NNX::FMLTensorDesc> OutputTensors, const UE::NNECore::FAttributeMap& Attributes) override
 		{
 			check(InputTensors.Num() >= 2 && InputTensors.Num() <= 3);
 			check(OutputTensors.Num() == 1);
diff --git a/Engine/Plugins/Experimental/NNX/Source/NNXRuntimeRDG/Private/Hlsl/NNIRuntimeRDGConvTranspose.cpp b/Engine/Plugins/Experimental/NNX/Source/NNXRuntimeRDG/Private/Hlsl/NNIRuntimeRDGConvTranspose.cpp
--- a/Engine/Plugins/Experimental/NNX/Source/NNXRuntimeRDG/Private/Hlsl/NNIRuntimeRDGConvTranspose.cpp
+++ b/Engine/Plugins/Experimental/NNX/Source/NNXRuntimeRDG/Private/Hlsl/NNIRuntimeRDGConvTranspose.cpp
@@ -46,7 +46,7 @@ namespace UE::NNIRuntimeRDG::Private::Hlsl
 
 	public:
 
-		virtual bool Initialize(TArrayView<const NNX::FMLTensorDesc> InputTensors, TArrayView<const NNX::FMLTensorDesc> OutputTensors, const FMLAttributeMap& Attributes) override
+		virtual bool Initialize(TArrayView<const NNX::FMLTensorDesc> InputTensors, TArrayView<const NNX::FMLTensorDesc> OutputTensors, const UE::NNECore::FAttributeMap& Attributes) override
 		{
 			check(InputTensors.Num() >= 2 && InputTensors.Num() <= 3);
 			check(OutputTensors.Num() == 1);
diff --git a/Engine/Plugins/Experimental/NNX/Source/NNXRuntimeRDG/Private/Hlsl/NNIRuntimeRDGElementWiseBinary.cpp b/Engine/Plugins/Experimental/NNX/Source/NNXRuntimeRDG/Private/Hlsl/NNIRuntimeRDGElementWiseBinary.cpp
--- a/Engine/Plugins/Experimental/NNX/Source/NNXRuntimeRDG/Private/Hlsl/NNIRuntimeRDGElementWiseBinary.cpp
+++ b/Engine/Plugins/Experimental/NNX/Source/NNXRuntimeRDG/Private/Hlsl/NNIRuntimeRDGElementWiseBinary.cpp
@@ -30,7 +30,7 @@ namespace UE::NNIRuntimeRDG::Private::Hlsl
 
 	public:
 
-		virtual bool Initialize(TArrayView<const NNX::FMLTensorDesc> InputTensors, TArrayView<const NNX::FMLTensorDesc> OutputTensors, const FMLAttributeMap& Attributes) override
+		virtual bool Initialize(TArrayView<const NNX::FMLTensorDesc> InputTensors, TArrayView<const NNX::FMLTensorDesc> OutputTensors, const UE::NNECore::FAttributeMap& Attributes) override
 		{
 			check(InputTensors.Num() == 2);
 			check(OutputTensors.Num() == 1);
diff --git a/Engine/Plugins/Experimental/NNX/Source/NNXRuntimeRDG/Private/Hlsl/NNIRuntimeRDGElementWiseUnary.cpp b/Engine/Plugins/Experimental/NNX/Source/NNXRuntimeRDG/Private/Hlsl/NNIRuntimeRDGElementWiseUnary.cpp
--- a/Engine/Plugins/Experimental/NNX/Source/NNXRuntimeRDG/Private/Hlsl/NNIRuntimeRDGElementWiseUnary.cpp
+++ b/Engine/Plugins/Experimental/NNX/Source/NNXRuntimeRDG/Private/Hlsl/NNIRuntimeRDGElementWiseUnary.cpp
@@ -3,6 +3,7 @@
 #include "NNIRuntimeRDGElementWiseUnary.h"
 #include "NNIHlslShadersElementWiseUnaryCS.h"
 #include "NNXRuntimeHLSLHelper.h"
+#include "NNECoreAttributeMap.h"
 
 namespace UE::NNIRuntimeRDG::Private::Hlsl
 {
@@ -32,7 +33,7 @@ namespace UE::NNIRuntimeRDG::Private::Hlsl
 
 	public:
 
-		virtual bool Initialize(TArrayView<const NNX::FMLTensorDesc> InputTensors, TArrayView<const NNX::FMLTensorDesc> OutputTensors, const FMLAttributeMap& Attributes) override
+		virtual bool Initialize(TArrayView<const NNX::FMLTensorDesc> InputTensors, TArrayView<const NNX::FMLTensorDesc> OutputTensors, const UE::NNECore::FAttributeMap& Attributes) override
 		{
 			check(InputTensors.Num() == 1);
 			check(OutputTensors.Num() == 1);
@@ -40,9 +41,9 @@ namespace UE::NNIRuntimeRDG::Private::Hlsl
 			Input = InputTensors[0];
 			Output = OutputTensors[0];
 
-			Alpha = Attributes.GetOptionalFloat(TEXT("alpha"), Alpha);
-			Beta = Attributes.GetOptionalFloat(TEXT("beta"), Beta);
-			Gamma = Attributes.GetOptionalFloat(TEXT("gamma"), Gamma);
+			Alpha = Attributes.GetValueOrDefault(TEXT("alpha"), Alpha);
+			Beta = Attributes.GetValueOrDefault(TEXT("beta"), Beta);
+			Gamma = Attributes.GetValueOrDefault(TEXT("gamma"), Gamma);
 
 			return true;
 		}
diff --git a/Engine/Plugins/Experimental/NNX/Source/NNXRuntimeRDG/Private/Hlsl/NNIRuntimeRDGElementWiseVariadic.cpp b/Engine/Plugins/Experimental/NNX/Source/NNXRuntimeRDG/Private/Hlsl/NNIRuntimeRDGElementWiseVariadic.cpp
--- a/Engine/Plugins/Experimental/NNX/Source/NNXRuntimeRDG/Private/Hlsl/NNIRuntimeRDGElementWiseVariadic.cpp
+++ b/Engine/Plugins/Experimental/NNX/Source/NNXRuntimeRDG/Private/Hlsl/NNIRuntimeRDGElementWiseVariadic.cpp
@@ -99,7 +99,7 @@ namespace UE::NNIRuntimeRDG::Private::Hlsl
 
 	public:
 
-		virtual bool Initialize(TArrayView<const NNX::FMLTensorDesc> InputTensors, TArrayView<const NNX::FMLTensorDesc> OutputTensors, const FMLAttributeMap& Attributes) override
+		virtual bool Initialize(TArrayView<const NNX::FMLTensorDesc> InputTensors, TArrayView<const NNX::FMLTensorDesc> OutputTensors, const UE::NNECore::FAttributeMap& Attributes) override
 		{
 			check(InputTensors.Num() > 0);
 			check(OutputTensors.Num() == 1);
diff --git a/Engine/Plugins/Experimental/NNX/Source/NNXRuntimeRDG/Private/Hlsl/NNIRuntimeRDGGather.cpp b/Engine/Plugins/Experimental/NNX/Source/NNXRuntimeRDG/Private/Hlsl/NNIRuntimeRDGGather.cpp
--- a/Engine/Plugins/Experimental/NNX/Source/NNXRuntimeRDG/Private/Hlsl/NNIRuntimeRDGGather.cpp
+++ b/Engine/Plugins/Experimental/NNX/Source/NNXRuntimeRDG/Private/Hlsl/NNIRuntimeRDGGather.cpp
@@ -3,6 +3,7 @@
 #include "NNIRuntimeRDGGather.h"
 #include "NNXRuntimeHLSLHelper.h"
 #include "NNIHlslShadersGatherCS.h"
+#include "NNECoreAttributeMap.h"
 
 namespace UE::NNIRuntimeRDG::Private::Hlsl
 {
@@ -28,7 +29,7 @@ namespace UE::NNIRuntimeRDG::Private::Hlsl
 
 	public:
 
-		virtual bool Initialize(TArrayView<const NNX::FMLTensorDesc> InputTensors, TArrayView<const NNX::FMLTensorDesc> OutputTensors, const FMLAttributeMap& Attributes) override
+		virtual bool Initialize(TArrayView<const NNX::FMLTensorDesc> InputTensors, TArrayView<const NNX::FMLTensorDesc> OutputTensors, const UE::NNECore::FAttributeMap& Attributes) override
 		{
 			using namespace UE::NNIHlslShaders::Internal;
 
@@ -39,7 +40,7 @@ namespace UE::NNIRuntimeRDG::Private::Hlsl
 			check(InputTensors[1].Shape.Num() > 0)
 			check(InputTensors[0].Shape.Num() + (InputTensors[1].Shape.Num() - 1) <= FGatherConstants::MAX_NUM_DIMENSIONS)
 
-			Axis = Attributes.GetOptionalFloat(TEXT("axis"), Axis);
+			Axis = Attributes.GetValueOrDefault(TEXT("axis"), Axis);
 			check(Axis < InputTensors[0].Shape.Num())
 			check(Axis >= -InputTensors[0].Shape.Num())
 			Axis = Axis >= 0 ? Axis : InputTensors[0].Shape.Num() + Axis;
diff --git a/Engine/Plugins/Experimental/NNX/Source/NNXRuntimeRDG/Private/Hlsl/NNIRuntimeRDGGemm.cpp b/Engine/Plugins/Experimental/NNX/Source/NNXRuntimeRDG/Private/Hlsl/NNIRuntimeRDGGemm.cpp
--- a/Engine/Plugins/Experimental/NNX/Source/NNXRuntimeRDG/Private/Hlsl/NNIRuntimeRDGGemm.cpp
+++ b/Engine/Plugins/Experimental/NNX/Source/NNXRuntimeRDG/Private/Hlsl/NNIRuntimeRDGGemm.cpp
@@ -3,6 +3,7 @@
 #include "NNIRuntimeRDGGemm.h"
 #include "NNIHlslShadersGemmCS.h"
 #include "NNXRuntimeHLSLHelper.h"
+#include "NNECoreAttributeMap.h"
 
 namespace UE::NNIRuntimeRDG::Private::Hlsl
 {
@@ -40,7 +41,7 @@ namespace UE::NNIRuntimeRDG::Private::Hlsl
 
 	public:
 
-		virtual bool Initialize(TArrayView<const NNX::FMLTensorDesc> InputTensors, TArrayView<const NNX::FMLTensorDesc> OutputTensors, const FMLAttributeMap& Attributes) override
+		virtual bool Initialize(TArrayView<const NNX::FMLTensorDesc> InputTensors, TArrayView<const NNX::FMLTensorDesc> OutputTensors, const UE::NNECore::FAttributeMap& Attributes) override
 		{
 			check(InputTensors.Num() >= 2 && InputTensors.Num() <= 3);
 			check(OutputTensors.Num() == 1);
@@ -61,10 +62,10 @@ namespace UE::NNIRuntimeRDG::Private::Hlsl
 			// CScalar = C != nullptr ? C[0] : (InElementType)0;
 			bNoBias = InputTensors.Num() != 3 /*|| InputC.Sizes[0] * InputC.Sizes[1] < 1*/;
 
-			InputAlpha = Attributes.GetOptionalFloat(TEXT("alpha"), InputAlpha);
-			InputBeta = Attributes.GetOptionalFloat(TEXT("beta"), InputBeta);
-			InputTransA = Attributes.GetOptionalInt32(TEXT("transA"), InputTransA);
-			InputTransB = Attributes.GetOptionalInt32(TEXT("transB"), InputTransB);
+			InputAlpha = Attributes.GetValueOrDefault(TEXT("alpha"), InputAlpha);
+			InputBeta = Attributes.GetValueOrDefault(TEXT("beta"), InputBeta);
+			InputTransA = Attributes.GetValueOrDefault(TEXT("transA"), InputTransA);
+			InputTransB = Attributes.GetValueOrDefault(TEXT("transB"), InputTransB);
 
 			return true;
 		}
@@ -108,15 +109,15 @@ namespace UE::NNIRuntimeRDG::Private::Hlsl
 		}
 	};
 
-	bool ValidateGemmOperator(const FMLAttributeMap& AttributeMap, TConstArrayView<EMLTensorDataType> InputTypes, TConstArrayView<const NNX::FSymbolicTensorShape> InputShapes)
+	bool ValidateGemmOperator(const UE::NNECore::FAttributeMap& AttributeMap, TConstArrayView<EMLTensorDataType> InputTypes, TConstArrayView<const NNX::FSymbolicTensorShape> InputShapes)
 	{
 		bool bIsValid = true;
 
 		NNX::FAttributeValidator AttributeValidator;
-		AttributeValidator.AddOptional(TEXT("alpha"), EMLAttributeDataType::Float);
-		AttributeValidator.AddOptional(TEXT("beta"), EMLAttributeDataType::Float);
-		AttributeValidator.AddOptional(TEXT("transA"), EMLAttributeDataType::Int32);
-		AttributeValidator.AddOptional(TEXT("transB"), EMLAttributeDataType::Int32);
+		AttributeValidator.AddOptional(TEXT("alpha"), ENNEAttributeDataType::Float);
+		AttributeValidator.AddOptional(TEXT("beta"), ENNEAttributeDataType::Float);
+		AttributeValidator.AddOptional(TEXT("transA"), ENNEAttributeDataType::Int32);
+		AttributeValidator.AddOptional(TEXT("transB"), ENNEAttributeDataType::Int32);
 		bIsValid &= AttributeValidator.Validate(AttributeMap);
 
 		NNX::FInputValidator InputValidator;
diff --git a/Engine/Plugins/Experimental/NNX/Source/NNXRuntimeRDG/Private/Hlsl/NNIRuntimeRDGMatMul.cpp b/Engine/Plugins/Experimental/NNX/Source/NNXRuntimeRDG/Private/Hlsl/NNIRuntimeRDGMatMul.cpp
--- a/Engine/Plugins/Experimental/NNX/Source/NNXRuntimeRDG/Private/Hlsl/NNIRuntimeRDGMatMul.cpp
+++ b/Engine/Plugins/Experimental/NNX/Source/NNXRuntimeRDG/Private/Hlsl/NNIRuntimeRDGMatMul.cpp
@@ -26,7 +26,7 @@ namespace UE::NNIRuntimeRDG::Private::Hlsl
 
 	public:
 
-		virtual bool Initialize(TArrayView<const NNX::FMLTensorDesc> InputTensors, TArrayView<const NNX::FMLTensorDesc> OutputTensors, const FMLAttributeMap& Attributes) override
+		virtual bool Initialize(TArrayView<const NNX::FMLTensorDesc> InputTensors, TArrayView<const NNX::FMLTensorDesc> OutputTensors, const UE::NNECore::FAttributeMap& Attributes) override
 		{
 			check(InputTensors.Num() == 2);
 			check(OutputTensors.Num() == 1);
diff --git a/Engine/Plugins/Experimental/NNX/Source/NNXRuntimeRDG/Private/NNXRuntimeDML.cpp b/Engine/Plugins/Experimental/NNX/Source/NNXRuntimeRDG/Private/NNXRuntimeDML.cpp
--- a/Engine/Plugins/Experimental/NNX/Source/NNXRuntimeRDG/Private/NNXRuntimeDML.cpp
+++ b/Engine/Plugins/Experimental/NNX/Source/NNXRuntimeRDG/Private/NNXRuntimeDML.cpp
@@ -6,6 +6,8 @@
 #include "NNXModelOptimizer.h"
 #include "NNXOperator.h"
 
+#include "NNECoreAttributeMap.h"
+
 #include "HAL/FileManager.h"
 
 #include "RenderGraphBuilder.h"
@@ -325,7 +327,7 @@ class FMLOperatorDml : public FMLOperatorRDG
 {
 public:
 
-	virtual bool Initialize(FDeviceContextDml* DevCtx, TArrayView<const FMLTensorDesc> InputTensors, TArrayView<const FMLTensorDesc> OutputTensors, const FMLAttributeMap& Attributes) = 0;
+	virtual bool Initialize(FDeviceContextDml* DevCtx, TArrayView<const FMLTensorDesc> InputTensors, TArrayView<const FMLTensorDesc> OutputTensors, const UE::NNECore::FAttributeMap& Attributes) = 0;
 
 	virtual void Dispatch(FRDGBuilder& GraphBuilder, TArrayView<const FMLTensorBinding> InInputBindings, TArrayView<const FMLTensorBinding> OutOutputBindings) = 0;
 
@@ -547,7 +549,7 @@ public:
 	//
 	//
 	//
-	virtual bool Initialize(FDeviceContextDml* InDevCtx, TArrayView<const FMLTensorDesc> InputTensors, TArrayView<const FMLTensorDesc> OutputTensors, const FMLAttributeMap& Attributes) override
+	virtual bool Initialize(FDeviceContextDml* InDevCtx, TArrayView<const FMLTensorDesc> InputTensors, TArrayView<const FMLTensorDesc> OutputTensors, const UE::NNECore::FAttributeMap& Attributes) override
 	{
 		// TODO: Setup attributes
 		Num = InputTensors[0].Num();
@@ -557,9 +559,9 @@ public:
 		const FMLTensorDesc& InputTensorDesc = InputTensors[0];
 		const FMLTensorDesc& OutputTensorDesc = OutputTensors[0];
 
-		Alpha = Attributes.GetOptionalFloat(TEXT("alpha"), Alpha);
-		Beta = Attributes.GetOptionalFloat(TEXT("beta"), Beta);
-		Gamma = Attributes.GetOptionalFloat(TEXT("gamma"), Gamma);
+		Alpha = Attributes.GetValueOrDefault(TEXT("alpha"), Alpha);
+		Beta = Attributes.GetValueOrDefault(TEXT("beta"), Beta);
+		Gamma = Attributes.GetValueOrDefault(TEXT("gamma"), Gamma);
 
 		// Initialize tensor descriptor (it's same for both input and output)
 		DmlUtil::FTensorDesc	DmlTensorDesc{};
@@ -797,7 +799,7 @@ public:
 	//
 	//
 	//
-	virtual bool Initialize(FDeviceContextDml* InDevCtx, TArrayView<const FMLTensorDesc> InputTensors, TArrayView<const FMLTensorDesc> OutputTensors, const FMLAttributeMap& Attributes) override
+	virtual bool Initialize(FDeviceContextDml* InDevCtx, TArrayView<const FMLTensorDesc> InputTensors, TArrayView<const FMLTensorDesc> OutputTensors, const UE::NNECore::FAttributeMap& Attributes) override
 	{
 		// TODO: Setup attributes
 		Num = OutputTensors[0].Num();
@@ -982,7 +984,7 @@ protected:
 
 private:
 
-	FMLOperatorDml* OpCreate(const FString& Name, TArrayView<const FMLTensorDesc> InputTensorDesc, TArrayView<const FMLTensorDesc> OutputTensorDescs, const FMLAttributeMap& Attributes);
+	FMLOperatorDml* OpCreate(const FString& Name, TArrayView<const FMLTensorDesc> InputTensorDesc, TArrayView<const FMLTensorDesc> OutputTensorDescs, const UE::NNECore::FAttributeMap& Attributes);
 	
 	TArray<FMLOperatorDml*>		Operators;
 	FDeviceContextDml*			DevCtx;
@@ -1280,7 +1282,7 @@ bool FMLInferenceModelDml::Init(UMLInferenceModel* InModel, FDeviceContextDml* I
 		
 		TArray<FMLTensorDesc> OpInputTensors;
 		TArray<FMLTensorDesc> OpOutputTensors;
-		FMLAttributeMap AttributeMap;
+		UE::NNECore::FAttributeMap AttributeMap;
 
 		for (int32 InputTensorIndex : Format.Operators[Idx].InTensors)
 		{
@@ -1288,7 +1290,7 @@ bool FMLInferenceModelDml::Init(UMLInferenceModel* InModel, FDeviceContextDml* I
 		}
 		for (int32 OutputTensorIndex : Format.Operators[Idx].OutTensors)
 		{
-			OpInputTensors.Emplace(AllTensors[OutputTensorIndex]);
+			OpOutputTensors.Emplace(AllTensors[OutputTensorIndex]);
 		}
 		for (const FMLFormatAttributeDesc& Desc : Format.Operators[Idx].Attributes)
 		{
@@ -1328,7 +1330,7 @@ void FMLInferenceModelDml::AddDispatchOps_RenderThread(FRDGBuilder& GraphBuilder
 //
 // Create operator
 //
-FMLOperatorDml* FMLInferenceModelDml::OpCreate(const FString& OpName, TArrayView<const FMLTensorDesc> InputTensorDescs, TArrayView<const FMLTensorDesc> OutputTensorDescs, const FMLAttributeMap& Attributes)
+FMLOperatorDml* FMLInferenceModelDml::OpCreate(const FString& OpName, TArrayView<const FMLTensorDesc> InputTensorDescs, TArrayView<const FMLTensorDesc> OutputTensorDescs, const UE::NNECore::FAttributeMap& Attributes)
 {
 	// TODO: Check if D3D12 device is valid
 
diff --git a/Engine/Plugins/Experimental/NNX/Source/NNXRuntimeRDG/Private/NNXRuntimeHLSL.cpp b/Engine/Plugins/Experimental/NNX/Source/NNXRuntimeRDG/Private/NNXRuntimeHLSL.cpp
--- a/Engine/Plugins/Experimental/NNX/Source/NNXRuntimeRDG/Private/NNXRuntimeHLSL.cpp
+++ b/Engine/Plugins/Experimental/NNX/Source/NNXRuntimeRDG/Private/NNXRuntimeHLSL.cpp
@@ -3,6 +3,7 @@
 #include "NNXInferenceModel.h"
 #include "NNXRuntimeFormat.h"
 #include "NNXModelOptimizer.h"
+#include "NNECoreAttributeMap.h"
 #include "Hlsl/NNIRuntimeRDGConv.h"
 #include "Hlsl/NNIRuntimeRDGElementWiseBinary.h"
 #include "Hlsl/NNIRuntimeRDGElementWiseUnary.h"
@@ -37,7 +38,7 @@ protected:
 
 private:
 
-	FMLOperatorHlsl* OpCreate(const FString& Name, TArrayView<const FMLTensorDesc> InputTensorDesc, TArrayView<const FMLTensorDesc> OutputTensorDescs, const FMLAttributeMap& Attributes);
+	FMLOperatorHlsl* OpCreate(const FString& Name, TArrayView<const FMLTensorDesc> InputTensorDesc, TArrayView<const FMLTensorDesc> OutputTensorDescs, const UE::NNECore::FAttributeMap& Attributes);
 
 	TArray<FMLOperatorHlsl*>	Operators;
 };
@@ -78,7 +79,7 @@ bool FMLInferenceModelHlsl::Init(UMLInferenceModel* InModel)
 
 		TArray<FMLTensorDesc> OpInputTensors;
 		TArray<FMLTensorDesc> OpOutputTensors;
-		FMLAttributeMap AttributeMap;
+		UE::NNECore::FAttributeMap AttributeMap;
 
 		for (int32 InputTensorIndex : Format.Operators[Idx].InTensors)
 		{
@@ -86,7 +87,7 @@ bool FMLInferenceModelHlsl::Init(UMLInferenceModel* InModel)
 		}
 		for (int32 OutputTensorIndex : Format.Operators[Idx].OutTensors)
 		{
-			OpInputTensors.Emplace(AllTensors[OutputTensorIndex]);
+			OpOutputTensors.Emplace(AllTensors[OutputTensorIndex]);
 		}
 		for (const FMLFormatAttributeDesc& Desc : Format.Operators[Idx].Attributes)
 		{
@@ -112,7 +113,7 @@ bool FMLInferenceModelHlsl::Init(UMLInferenceModel* InModel)
 //
 //
 //
-FMLOperatorHlsl* FMLInferenceModelHlsl::OpCreate(const FString& OpName, TArrayView<const FMLTensorDesc> InputTensorDescs, TArrayView<const FMLTensorDesc> OutputTensorDescs, const FMLAttributeMap& AttributeMap)
+FMLOperatorHlsl* FMLInferenceModelHlsl::OpCreate(const FString& OpName, TArrayView<const FMLTensorDesc> InputTensorDescs, TArrayView<const FMLTensorDesc> OutputTensorDescs, const UE::NNECore::FAttributeMap& AttributeMap)
 {
 	FMLOperatorRegistryHlsl::OperatorCreateFunc CreateFn = FMLOperatorRegistryHlsl::Get()->OpFind(OpName);
 
diff --git a/Engine/Plugins/Experimental/NNX/Source/NNXRuntimeRDG/Private/NNXRuntimeHLSLOp.h b/Engine/Plugins/Experimental/NNX/Source/NNXRuntimeRDG/Private/NNXRuntimeHLSLOp.h
--- a/Engine/Plugins/Experimental/NNX/Source/NNXRuntimeRDG/Private/NNXRuntimeHLSLOp.h
+++ b/Engine/Plugins/Experimental/NNX/Source/NNXRuntimeRDG/Private/NNXRuntimeHLSLOp.h
@@ -13,7 +13,7 @@ namespace NNX
 	struct FMLOperatorHlsl : public FMLOperatorRDG
 	{
 		virtual ~FMLOperatorHlsl() = default;
-		virtual bool Initialize(TArrayView<const FMLTensorDesc> InputTensors, TArrayView<const FMLTensorDesc> OutputTensors, const FMLAttributeMap& Attributes) = 0;
+		virtual bool Initialize(TArrayView<const FMLTensorDesc> InputTensors, TArrayView<const FMLTensorDesc> OutputTensors, const UE::NNECore::FAttributeMap& Attributes) = 0;
 		virtual void Dispatch(FRDGBuilder& GraphBuilder, TArrayView<const FMLTensorBinding> InInputBindings, TArrayView<const FMLTensorBinding> OutOutputBindings) = 0;
 	};
 
diff --git a/Engine/Plugins/Experimental/NNX/Source/NNXRuntimeRDG/Private/NNXRuntimeRDG.cpp b/Engine/Plugins/Experimental/NNX/Source/NNXRuntimeRDG/Private/NNXRuntimeRDG.cpp
--- a/Engine/Plugins/Experimental/NNX/Source/NNXRuntimeRDG/Private/NNXRuntimeRDG.cpp
+++ b/Engine/Plugins/Experimental/NNX/Source/NNXRuntimeRDG/Private/NNXRuntimeRDG.cpp
@@ -16,7 +16,7 @@ namespace NNX
 //
 //
 bool AlwaysValidValidationFunction(
-	const FMLAttributeMap& AttributeMap, 
+	const UE::NNECore::FAttributeMap& AttributeMap, 
 	TConstArrayView<EMLTensorDataType> InputTensorTypes,
 	TConstArrayView<const FSymbolicTensorShape> InputShapes)
 {
@@ -90,28 +90,28 @@ void FInputValidator::AddRequired(int32 TemplateIdx)
 //
 //
 //
-void FAttributeValidator::AddOptional(const FString& Name, EMLAttributeDataType Type)
+void FAttributeValidator::AddOptional(const FString& Name, ENNEAttributeDataType Type)
 {
 	checkf(nullptr == OptionalAttributes.FindByPredicate([Name](const FEntry& Other) { return Other.Name == Name; }), TEXT("Attribute name should be unique"));
 	checkf(nullptr == RequiredAttributes.FindByPredicate([Name](const FEntry& Other) { return Other.Name == Name; }), TEXT("Attribute name should be unique"));
 	OptionalAttributes.Emplace(Name, Type);
 }
 
-void FAttributeValidator::AddRequired(const FString& Name, EMLAttributeDataType Type)
+void FAttributeValidator::AddRequired(const FString& Name, ENNEAttributeDataType Type)
 {
 	checkf(nullptr == OptionalAttributes.FindByPredicate([Name](const FEntry& Other) { return Other.Name == Name; }), TEXT("Attribute name should be unique"));
 	checkf(nullptr == RequiredAttributes.FindByPredicate([Name](const FEntry& Other) { return Other.Name == Name; }), TEXT("Attribute name should be unique"));
 	RequiredAttributes.Emplace(Name, Type);
 }
 
-bool FAttributeValidator::Validate(const FMLAttributeMap& AttributesToValidate)
+bool FAttributeValidator::Validate(const UE::NNECore::FAttributeMap& AttributesToValidate)
 {
 	bool bAreAttributesValid = true;
 
 	//Verify all required attribute are matching specifications
 	for (int32 Idx = 0; Idx < RequiredAttributes.Num(); ++Idx)
 	{
-		const FMLAttributeValue* FoundAttribute = AttributesToValidate.GetAttributeValue(RequiredAttributes[Idx].Name);
+		const FNNEAttributeValue* FoundAttribute = AttributesToValidate.GetAttributeValue(RequiredAttributes[Idx].Name);
 		
 		if (FoundAttribute == nullptr)
 		{
@@ -132,7 +132,7 @@ bool FAttributeValidator::Validate(const FMLAttributeMap& AttributesToValidate)
 	//Verify all optional attribute are matching specifications
 	for (int32 Idx = 0; Idx < OptionalAttributes.Num(); ++Idx)
 	{
-		const FMLAttributeValue* FoundAttribute = AttributesToValidate.GetAttributeValue(OptionalAttributes[Idx].Name);
+		const FNNEAttributeValue* FoundAttribute = AttributesToValidate.GetAttributeValue(OptionalAttributes[Idx].Name);
 		
 		if ((FoundAttribute != nullptr) && (OptionalAttributes[Idx].Type != FoundAttribute->GetType()))
 		{
diff --git a/Engine/Plugins/Experimental/NNX/Source/NNXRuntimeRDG/Private/NNXRuntimeRDG.h b/Engine/Plugins/Experimental/NNX/Source/NNXRuntimeRDG/Private/NNXRuntimeRDG.h
--- a/Engine/Plugins/Experimental/NNX/Source/NNXRuntimeRDG/Private/NNXRuntimeRDG.h
+++ b/Engine/Plugins/Experimental/NNX/Source/NNXRuntimeRDG/Private/NNXRuntimeRDG.h
@@ -2,12 +2,15 @@
 
 #pragma once
 
+#include "NNECoreAttributeMap.h"
+#include "NNECoreAttributeValue.h"
 #include "NNXCore.h"
+#include "NNXModelOptimizerInterface.h"
 #include "NNXRuntime.h"
-#include "NNXModelOptimizer.h"
-#include "ShaderParameterUtils.h"
+#include "NNXRuntimeFormat.h"
 #include "RHIGPUReadback.h"
 #include "Serialization/MemoryReader.h"
+#include "ShaderParameterUtils.h"
 
 #include "Containers/Map.h"
 
@@ -82,7 +85,7 @@ protected:
 
 //TODO jira 167585 remove default validation and declare contract in all HLSL operator (see HLSL Gemm for current example)
 //TODO jira 167584 remove default validation and declare contract in all DML operator (see HLSL Gemm for current example)
-bool AlwaysValidValidationFunction(const FMLAttributeMap& AttributeMap, TConstArrayView<EMLTensorDataType> InputTypes, TConstArrayView<const FSymbolicTensorShape> InputShapes);
+bool AlwaysValidValidationFunction(const UE::NNECore::FAttributeMap& AttributeMap, TConstArrayView<EMLTensorDataType> InputTypes, TConstArrayView<const FSymbolicTensorShape> InputShapes);
 
 class FInputValidator
 {
@@ -104,14 +107,14 @@ private:
 class FAttributeValidator
 {
 public:
-	void AddOptional(const FString& Name, EMLAttributeDataType Type);
-	void AddRequired(const FString& Name, EMLAttributeDataType Type);
-	bool Validate(const FMLAttributeMap& AttributesToValidate);
+	void AddOptional(const FString& Name, ENNEAttributeDataType Type);
+	void AddRequired(const FString& Name, ENNEAttributeDataType Type);
+	bool Validate(const UE::NNECore::FAttributeMap& AttributesToValidate);
 
 private:
 	struct FEntry
 	{
-		FEntry(const FString& InName, EMLAttributeDataType InType)
+		FEntry(const FString& InName, ENNEAttributeDataType InType)
 			: Name(InName), Type(InType)
 		{
 		}
@@ -121,7 +124,7 @@ private:
 		//In the same direction we might only support a range of value for a float (for example
 		//we only support integer but the type is float, or only positive values for an int32)
 		FString Name;
-		EMLAttributeDataType Type;
+		ENNEAttributeDataType Type;
 	};
 
 	TArray<FEntry> RequiredAttributes;
@@ -137,7 +140,7 @@ class TOperatorRegistryRDG
 public:
 
 	typedef TOperatorType* (*OperatorCreateFunc)();
-	typedef bool (*OperatorValidateFunc)(const FMLAttributeMap& AttributeMap, TConstArrayView<EMLTensorDataType> InputTypes, TConstArrayView<const FSymbolicTensorShape> InputShapes);
+	typedef bool (*OperatorValidateFunc)(const UE::NNECore::FAttributeMap& AttributeMap, TConstArrayView<EMLTensorDataType> InputTypes, TConstArrayView<const FSymbolicTensorShape> InputShapes);
 
 	static TOperatorRegistryRDG* Get()
 	{
@@ -224,7 +227,7 @@ public:
 		{
 			TArray<EMLTensorDataType> InputTensorTypes;
 			TArray<FSymbolicTensorShape> InputTensorShapes;
-			FMLAttributeMap AttributeMap;
+			UE::NNECore::FAttributeMap AttributeMap;
 			
 			for (int32 InputTensorIndex: Format.Operators[Idx].InTensors)
 			{
diff --git a/Engine/Plugins/Experimental/NNX/Source/NNXUtils/Private/NNXModelBuilderNNX.cpp b/Engine/Plugins/Experimental/NNX/Source/NNXUtils/Private/NNXModelBuilderNNX.cpp
--- a/Engine/Plugins/Experimental/NNX/Source/NNXUtils/Private/NNXModelBuilderNNX.cpp
+++ b/Engine/Plugins/Experimental/NNX/Source/NNXUtils/Private/NNXModelBuilderNNX.cpp
@@ -3,6 +3,7 @@
 #include "NNXModelBuilder.h"
 #include "NNXRuntimeFormat.h"
 #include "NNXCore.h"
+#include "NNECoreAttributeMap.h"
 
 #include "Misc/StringBuilder.h"
 #include "Serialization/MemoryWriter.h"
@@ -218,7 +219,7 @@ public:
 	}
 
 	/** Add operator attribute */
-	virtual bool AddOperatorAttribute(HOperator Op, const FString& Name, const FMLAttributeValue& Value) override
+	virtual bool AddOperatorAttribute(HOperator Op, const FString& Name, const FNNEAttributeValue& Value) override
 	{
 		int OpIdx = NNXOperatorCast(Op);
 
diff --git a/Engine/Plugins/Experimental/NNX/Source/NNXUtils/Private/NNXModelBuilderONNX.cpp b/Engine/Plugins/Experimental/NNX/Source/NNXUtils/Private/NNXModelBuilderONNX.cpp
--- a/Engine/Plugins/Experimental/NNX/Source/NNXUtils/Private/NNXModelBuilderONNX.cpp
+++ b/Engine/Plugins/Experimental/NNX/Source/NNXUtils/Private/NNXModelBuilderONNX.cpp
@@ -3,6 +3,7 @@
 #include "NNXCore.h"
 #include "NNXModelBuilder.h"
 #include "NNXRuntimeFormat.h"
+#include "NNECoreAttributeMap.h"
 
 #include "NNXThirdPartyWarningDisabler.h"
 NNX_THIRD_PARTY_INCLUDES_START
@@ -201,21 +202,21 @@ public:
 	}
 
 	/** Add operator attribute */
-	virtual bool AddOperatorAttribute(HOperator Op, const FString& Name, const FMLAttributeValue& Value) override
+	virtual bool AddOperatorAttribute(HOperator Op, const FString& Name, const FNNEAttributeValue& Value) override
 	{
 		auto NodeOp = OnnxOperatorCast(Op);
 		onnx::AttributeProto* Attribute = NodeOp->mutable_attribute()->Add();
 
 		Attribute->set_name(TCHAR_TO_ANSI(*Name));
-		if (Value.GetType() == EMLAttributeDataType::Float)
+		if (Value.GetType() == ENNEAttributeDataType::Float)
 		{
 			Attribute->set_type(onnx::AttributeProto::FLOAT);
-			Attribute->set_f(Value.AsFloat());
+			Attribute->set_f(Value.GetValue<float>());
 		}
-		else if (Value.GetType() == EMLAttributeDataType::Int32)
+		else if (Value.GetType() == ENNEAttributeDataType::Int32)
 		{
 			Attribute->set_type(onnx::AttributeProto::INT);
-			Attribute->set_i(Value.AsInt32());
+			Attribute->set_i(Value.GetValue<int32>());
 		}
 		else
 		{
@@ -277,14 +278,14 @@ private:
 //
 NNXUTILS_API bool CreateONNXModelForOperator(const FString& OperatorName, TConstArrayView<FMLTensorDesc> InInputTensors, TConstArrayView<FMLTensorDesc> InOutputTensors, FNNIModelRaw& Model)
 {
-	FMLAttributeMap EmptyAttributeMap;
+	UE::NNECore::FAttributeMap EmptyAttributeMap;
 	return CreateONNXModelForOperator(OperatorName, InInputTensors, InOutputTensors, EmptyAttributeMap, Model);
 }
 
 //
 //
 //
-NNXUTILS_API bool CreateONNXModelForOperator(const FString& OperatorName, TConstArrayView<FMLTensorDesc> InInputTensors, TConstArrayView<FMLTensorDesc> InOutputTensors, const FMLAttributeMap& Attributes, FNNIModelRaw& Model)
+NNXUTILS_API bool CreateONNXModelForOperator(const FString& OperatorName, TConstArrayView<FMLTensorDesc> InInputTensors, TConstArrayView<FMLTensorDesc> InOutputTensors, const UE::NNECore::FAttributeMap& Attributes, FNNIModelRaw& Model)
 {
 	Model = FNNIModelRaw{};
 	
diff --git a/Engine/Plugins/Experimental/NNX/Source/NNXUtils/Private/NNXModelOptimizer.cpp b/Engine/Plugins/Experimental/NNX/Source/NNXUtils/Private/NNXModelOptimizer.cpp
--- a/Engine/Plugins/Experimental/NNX/Source/NNXUtils/Private/NNXModelOptimizer.cpp
+++ b/Engine/Plugins/Experimental/NNX/Source/NNXUtils/Private/NNXModelOptimizer.cpp
@@ -1,6 +1,7 @@
 // Copyright Epic Games, Inc. All Rights Reserved.
 
 #include "NNXModelOptimizer.h"
+#include "NNXModelOptimizerInterface.h"
 #include "NNXCore.h"
 #include "NNXRuntimeFormat.h"
 #include "NNXModelBuilder.h"
@@ -543,11 +544,11 @@ private:
 
 				if (AttrInfo.type == Ort::GraphAttributeType::kFloat)
 				{
-					Builder->AddOperatorAttribute(Op, AttrInfo.name, FMLAttributeValue(AttrValue.f));
+					Builder->AddOperatorAttribute(Op, AttrInfo.name, FNNEAttributeValue(AttrValue.f));
 				}
 				else if (AttrInfo.type == Ort::GraphAttributeType::kInt)
 				{
-					Builder->AddOperatorAttribute(Op, AttrInfo.name, FMLAttributeValue((int)AttrValue.i));
+					Builder->AddOperatorAttribute(Op, AttrInfo.name, FNNEAttributeValue((int)AttrValue.i));
 				}
 				else
 				{
@@ -638,7 +639,7 @@ NNXUTILS_API TUniquePtr<IModelOptimizer> CreateModelOptimizer(ENNXInferenceForma
 }
 
 /** Helper to create an optimized model for a given runtime from ONNX */
-NNXUTILS_API bool CreateRuntimeModelFromONNX(const FNNIModelRaw& ONNXModel, FNNIModelRaw& OptimizedModel, FString RuntimeName, const FOptimizerOptionsMap& Options)
+NNXUTILS_API bool CreateRuntimeModelFromONNX(const FNNIModelRaw& ONNXModel, FNNIModelRaw& OptimizedModel, const FString& RuntimeName, const FOptimizerOptionsMap& Options)
 {
 	OptimizedModel = FNNIModelRaw{};
 	
diff --git a/Engine/Plugins/Experimental/NNX/Source/NNXUtils/Private/NNXUtilsModule.cpp b/Engine/Plugins/Experimental/NNX/Source/NNXUtils/Private/NNXUtilsModule.cpp
--- a/Engine/Plugins/Experimental/NNX/Source/NNXUtils/Private/NNXUtilsModule.cpp
+++ b/Engine/Plugins/Experimental/NNX/Source/NNXUtils/Private/NNXUtilsModule.cpp
@@ -1,9 +1,12 @@
 // Copyright Epic Games, Inc. All Rights Reserved.
 #include "Modules/ModuleManager.h"
 
+#include "NNECoreAttributeMap.h"
 #include "NNXCore.h"
-#include "NNXModelOptimizer.h"
 #include "NNXModelBuilder.h"
+#include "NNXModelOptimizer.h"
+#include "NNXModelOptimizerInterface.h"
+#include "NNXRuntimeFormat.h"
 
 #include "CoreMinimal.h"
 #include "Containers/Array.h"
diff --git a/Engine/Plugins/Experimental/NNX/Source/NNXUtils/Public/NNXModelBuilder.h b/Engine/Plugins/Experimental/NNX/Source/NNXUtils/Public/NNXModelBuilder.h
--- a/Engine/Plugins/Experimental/NNX/Source/NNXUtils/Public/NNXModelBuilder.h
+++ b/Engine/Plugins/Experimental/NNX/Source/NNXUtils/Public/NNXModelBuilder.h
@@ -2,8 +2,9 @@
 #pragma once
 
 #include "NNXTypes.h"
-#include "Containers/Array.h"
 
+namespace UE::NNECore { class FAttributeMap; }
+struct FNNEAttributeValue;
 struct FNNIModelRaw;
 
 namespace NNX
@@ -66,7 +67,7 @@ public:
 	virtual bool AddOperatorInput(HOperator Op, HTensor Tensor) = 0;
 
 	/** Add operator attribute */
-	virtual bool AddOperatorAttribute(HOperator Op, const FString& Name, const FMLAttributeValue& Value) = 0;
+	virtual bool AddOperatorAttribute(HOperator Op, const FString& Name, const FNNEAttributeValue& Value) = 0;
 	
 	/** Add operator output */
 	virtual bool AddOperatorOutput(HOperator Op, HTensor Tensor) = 0;
@@ -87,7 +88,7 @@ NNXUTILS_API IMLModelBuilder* CreateONNXModelBuilder(int64 IrVersion = OnnxIrVer
  * Utility functions to create single layer NN for operator testing with optional attributes
  */
 NNXUTILS_API bool CreateONNXModelForOperator(const FString& OperatorName, TConstArrayView<const FMLTensorDesc> InInputTensors, TConstArrayView<const FMLTensorDesc> InOutputTensors, FNNIModelRaw& ModelData);
-NNXUTILS_API bool CreateONNXModelForOperator(const FString& OperatorName, TConstArrayView<const FMLTensorDesc> InInputTensors, TConstArrayView<const FMLTensorDesc> InOutputTensors, const FMLAttributeMap& Attributes, FNNIModelRaw& ModelData);
+NNXUTILS_API bool CreateONNXModelForOperator(const FString& OperatorName, TConstArrayView<const FMLTensorDesc> InInputTensors, TConstArrayView<const FMLTensorDesc> InOutputTensors, const UE::NNECore::FAttributeMap& Attributes, FNNIModelRaw& ModelData);
 
 /**
  * Create an instance of NNX model builder that creates NNX model/format in memory
diff --git a/Engine/Plugins/Experimental/NNX/Source/NNXUtils/Public/NNXModelOptimizer.h b/Engine/Plugins/Experimental/NNX/Source/NNXUtils/Public/NNXModelOptimizer.h
--- a/Engine/Plugins/Experimental/NNX/Source/NNXUtils/Public/NNXModelOptimizer.h
+++ b/Engine/Plugins/Experimental/NNX/Source/NNXUtils/Public/NNXModelOptimizer.h
@@ -2,11 +2,14 @@
 #pragma once
 
 #include "NNXRuntimeFormat.h"
-#include "NNXModelOptimizerInterface.h"
+
+namespace UE::NNECore { class FAttributeMap; }
 
 namespace NNX
 {
+
 class IModelOptimizer;
+using FOptimizerOptionsMap = UE::NNECore::FAttributeMap;
 	
 /** Create a model optimizer */
 NNXUTILS_API TUniquePtr<IModelOptimizer> CreateModelOptimizer(ENNXInferenceFormat InputFormat, ENNXInferenceFormat OutputFormat);
@@ -27,6 +30,6 @@ inline TUniquePtr<IModelOptimizer> CreateONNXToONNXModelOptimizer()
 }
 
 /** Helper to create an optimized model for a given runtime from ONNX */
-NNXUTILS_API bool CreateRuntimeModelFromONNX(FNNIModelRaw& OutputModel, const FNNIModelRaw& ONNXModel, FString RuntimeName, const FOptimizerOptionsMap& Options);
+NNXUTILS_API bool CreateRuntimeModelFromONNX(FNNIModelRaw& OutputModel, const FNNIModelRaw& ONNXModel, const FString& RuntimeName, const FOptimizerOptionsMap& Options);
 
 } // NNX
