NNI - Multi layer support for fully static model

Problems:
* NNX does not have multi layer support
* NNX miss the framework to support dynamic input shape

Solutions:
* Add symbolic tensor description to model definition.
* Add SetInput to model definition to allow model to run shape inference.
* Fake shape inference (all model shapes should be concrete for now).
* Allocate intermediate tensor in RDG + provide the local input/output to each op

Tests:
* Run all tests on windows editor

diff --git a/Engine/Plugins/Experimental/NNX/Source/NNXCore/Private/NNXRuntime.cpp b/Engine/Plugins/Experimental/NNX/Source/NNXCore/Private/NNXRuntime.cpp
new file mode 100644
--- /dev/null
+++ b/Engine/Plugins/Experimental/NNX/Source/NNXCore/Private/NNXRuntime.cpp
@@ -0,0 +1,69 @@
+// Copyright Epic Games, Inc. All Rights Reserved.
+
+#include "NNXRuntime.h"
+#include "NNXCore.h"
+
+namespace NNX
+{
+
+	FMLTensorBinding FMLTensorBinding::FromCPU(void* CpuMemory, uint64 InSize, uint64 InOffset)
+	{
+		FMLTensorBinding	Binding;
+
+		Binding.BindingType = EMLTensorBindingDataType::CPUMemory;
+		Binding.CpuMemory = CpuMemory;
+		Binding.SizeInBytes = InSize;
+		Binding.OffsetInBytes = InOffset;
+
+		return Binding;
+	}
+
+	FMLTensorBinding FMLTensorBinding::FromGPU(uint64 GpuMemory, uint64 InSize, uint64 InOffset)
+	{
+		FMLTensorBinding	Binding;
+
+		Binding.BindingType = EMLTensorBindingDataType::GPUMemory;
+		Binding.GpuMemory = GpuMemory;
+		Binding.SizeInBytes = InSize;
+		Binding.OffsetInBytes = InOffset;
+
+		return Binding;
+	}
+
+	FMLTensorBinding FMLTensorBinding::FromRDG(FRDGBuffer* BufferRef, uint64 InSize, uint64 InOffset)
+	{
+		FMLTensorBinding	Binding;
+
+		Binding.BindingType = EMLTensorBindingDataType::RDGBuffer;
+		Binding.Buffer = BufferRef;
+		Binding.SizeInBytes = InSize;
+		Binding.OffsetInBytes = InOffset;
+
+		return Binding;
+	}
+
+	int FMLInferenceModel::SetInputShapes(TConstArrayView<const FTensorShape> InInputShapes)
+	{
+		if (InInputShapes.Num() != InputSymbolicTensors.Num())
+		{
+			UE_LOG(LogNNX, Warning, TEXT("Number of input shapes does not match number of input tensors"));
+			return -1;
+		}
+
+		for (int32 i = 0; i < InInputShapes.Num(); ++i)
+		{
+			const FSymbolicTensorDesc SymbolicDesc = InputSymbolicTensors[i];
+			if (!InInputShapes[i].IsCompatibleWith(SymbolicDesc.Shape))
+			{
+				UE_LOG(LogNNX, Warning, TEXT("Input shape does not match input tensor %s of index %d"), *SymbolicDesc.Name, i);
+				return -1;
+			}
+		}
+
+		//Implementations are responsible to handle shape inference.
+		//This base implementation only validate that all inputs are
+		//matching what the model can support.
+		return 0;
+	}
+
+}
\ No newline at end of file
diff --git a/Engine/Plugins/Experimental/NNX/Source/NNXCore/Public/NNXRuntime.h b/Engine/Plugins/Experimental/NNX/Source/NNXCore/Public/NNXRuntime.h
--- a/Engine/Plugins/Experimental/NNX/Source/NNXCore/Public/NNXRuntime.h
+++ b/Engine/Plugins/Experimental/NNX/Source/NNXCore/Public/NNXRuntime.h
@@ -3,9 +3,7 @@
 #pragma once
 
 #include "NNXTypes.h"
-
 #include "CoreMinimal.h"
-#include "Modules/ModuleInterface.h"
 
 class FRDGBuffer;
 
@@ -65,7 +63,7 @@ enum class EMLTensorBindingDataType : uint8
 };
 
 /** Tensor binding */
-struct FMLTensorBinding
+struct NNXCORE_API FMLTensorBinding
 {
 	union
 	{
@@ -80,50 +78,9 @@ struct FMLTensorBinding
 	uint64						OffsetInBytes;						//!< Offset in bytes from the start of data
 	EMLTensorBindingDataType	BindingType;
 
-	/**
-	 * Initialize binding from CPU memory
-	 */
-	static FMLTensorBinding FromCPU(void* CpuMemory, uint64 InSize, uint64 InOffset = 0)
-	{
-		FMLTensorBinding	Binding;
-
-		Binding.BindingType = EMLTensorBindingDataType::CPUMemory;
-		Binding.CpuMemory = CpuMemory;
-		Binding.SizeInBytes = InSize;
-		Binding.OffsetInBytes = InOffset;
-
-		return Binding;
-	}
-
-	/**
-	 * Initialize binding from GPU memory
-	 */
-	static FMLTensorBinding FromGPU(uint64 GpuMemory, uint64 InSize, uint64 InOffset = 0)
-	{
-		FMLTensorBinding	Binding;
-
-		Binding.BindingType = EMLTensorBindingDataType::GPUMemory;
-		Binding.GpuMemory = GpuMemory;
-		Binding.SizeInBytes = InSize;
-		Binding.OffsetInBytes = InOffset;
-
-		return Binding;
-	}
-
-	/**
-	 * Initialize binding from RDG allocated buffer memory
-	 */
-	static FMLTensorBinding FromRDG(FRDGBuffer* BufferRef, uint64 InSize, uint64 InOffset = 0)
-	{
-		FMLTensorBinding	Binding;
-
-		Binding.BindingType = EMLTensorBindingDataType::RDGBuffer;
-		Binding.Buffer = BufferRef;
-		Binding.SizeInBytes = InSize;
-		Binding.OffsetInBytes = InOffset;
-
-		return Binding;
-	}
+	static FMLTensorBinding FromCPU(void* CpuMemory, uint64 InSize, uint64 InOffset = 0);
+	static FMLTensorBinding FromGPU(uint64 GpuMemory, uint64 InSize, uint64 InOffset = 0);
+	static FMLTensorBinding FromRDG(FRDGBuffer* BufferRef, uint64 InSize, uint64 InOffset = 0);
 };
 
 enum class EMLInferenceModelType : uint8
@@ -134,7 +91,7 @@ enum class EMLInferenceModelType : uint8
 };
 
 /** Runtime Inference model is used to execute / run model inference */
-class FMLInferenceModel
+class NNXCORE_API FMLInferenceModel
 {
 public:
 
@@ -145,22 +102,30 @@ public:
 		return Type;
 	}
 
-	int32 InputTensorNum() const;
-	const FMLTensorDesc& GetInputTensor(int32 Index) const;
-	TArrayView<const FMLTensorDesc> GetInputTensors() const;
+	/** Getters for tensor description as defined by the model potentially with variable dimensions() */
+	TConstArrayView<const FSymbolicTensorDesc> GetInputTensorDescs() const;
+	TConstArrayView<const FSymbolicTensorDesc> GetOutputTensorDescs() const;
 
-	int32 OutputTensorNum() const;
-	const FMLTensorDesc& GetOutputTensor(int32 Index) const;
-	TArrayView<const FMLTensorDesc> GetOutputTensors() const;
-
-	/** This call is synchronous on all Inference model types(CPU, RDG, GPU), i.e.the calling thread will be blocked
-	 * until inference is not finished
-	 */
-	virtual int Run(TArrayView<const FMLTensorBinding> InInputTensors, TArrayView<const FMLTensorBinding> OutOutputTensors) = 0;
-
-	/** Enqueue the inference operators on the Render graph render thread. It's caller's responsibility to actually run the graph.
-	 */
-	virtual int EnqueueRDG(FRDGBuilder& GraphBuilder, TArrayView<const FMLTensorBinding> InInputTensors, TArrayView<const FMLTensorBinding> OutOutputTensors)
+	/** Getters for input/output shapes if they were set already. Empty list otherwise. See SetInputShapes() */
+	TConstArrayView<const FTensorShape> GetInputTensorShapes() const;
+	TConstArrayView<const FTensorShape> GetOutputTensorShapes() const;
+	
+	/** This call will prepare the model to be run with the given input shape, it is an optional call.
+	 * Run() and EnqueueRDG() will call it internally if it was not call explicitely.
+	 * It is recommended to call this function ahead of time to avoid overhead during the first model execution. */
+	virtual int SetInputShapes(TConstArrayView<const FTensorShape> InInputShapes);
+
+	/** This call is synchronous on all engine types(CPU, RDG, GPU), i.e.the calling thread will be blocked.
+	 * until execution is finished */
+	virtual int Run(TConstArrayView<const FMLTensorBinding> InInputTensors,
+		            TConstArrayView<const FTensorShape> InInputShapes,
+					TConstArrayView<const FMLTensorBinding> InOutputTensors) = 0;
+
+	/** Enqueue the execution on the Render graph render thread. It's caller's responsibility to actually run the graph. */
+	virtual int EnqueueRDG(FRDGBuilder& RDGBuilder, 
+		TConstArrayView<const FMLTensorBinding> InInputTensors, 
+		TConstArrayView<const FTensorShape> InInputShapes,
+		TConstArrayView<const FMLTensorBinding> InOutputTensors)
 	{
 		return -1;
 	}
@@ -171,9 +136,10 @@ protected:
 		: Type(InType)
 	{}
 
-	TArray<FMLTensorDesc>	InputTensors;
-	TArray<FMLTensorDesc>	OutputTensors;
-	TArray<FMLTensorDesc>	AllTensors;
+	TArray<FTensorShape>		InputTensorShapes;
+	TArray<FTensorShape>		OutputTensorShapes;
+	TArray<FSymbolicTensorDesc>	InputSymbolicTensors;
+	TArray<FSymbolicTensorDesc>	OutputSymbolicTensors;
 	EMLInferenceModelType	Type;
 };
 
@@ -182,36 +148,24 @@ protected:
 // IMPLEMENTATION
 //-----------------------------------------------------------------------------
 
-inline int32 FMLInferenceModel::InputTensorNum() const
-{
-	return InputTensors.Num();
-}
-
-inline const FMLTensorDesc& FMLInferenceModel::GetInputTensor(int32 Index) const
-{
-	checkf(Index < InputTensors.Num(), TEXT("Invalid tensor index"));
-	return InputTensors[Index];
-}
-
-inline TArrayView<const FMLTensorDesc> FMLInferenceModel::GetInputTensors() const
+inline TConstArrayView<const FSymbolicTensorDesc> FMLInferenceModel::GetInputTensorDescs() const
 {
-	return InputTensors;
+	return InputSymbolicTensors;
 }
 
-inline int32 FMLInferenceModel::OutputTensorNum() const
+inline TConstArrayView<const FSymbolicTensorDesc> FMLInferenceModel::GetOutputTensorDescs() const
 {
-	return OutputTensors.Num();
+	return OutputSymbolicTensors;
 }
 
-inline const FMLTensorDesc& FMLInferenceModel::GetOutputTensor(int32 Index) const
+inline TConstArrayView<const FTensorShape> FMLInferenceModel::GetInputTensorShapes() const
 {
-	checkf(Index < OutputTensors.Num(), TEXT("Invalid tensor index"));
-	return OutputTensors[Index];
+	return InputTensorShapes;
 }
 
-inline TArrayView<const FMLTensorDesc> FMLInferenceModel::GetOutputTensors() const
+inline TConstArrayView<const FTensorShape> FMLInferenceModel::GetOutputTensorShapes() const
 {
-	return OutputTensors;
+	return OutputTensorShapes;
 }
 
 //struct FMLOperatorDomain
diff --git a/Engine/Plugins/Experimental/NNX/Source/NNXCore/Public/NNXTypes.h b/Engine/Plugins/Experimental/NNX/Source/NNXCore/Public/NNXTypes.h
--- a/Engine/Plugins/Experimental/NNX/Source/NNXCore/Public/NNXTypes.h
+++ b/Engine/Plugins/Experimental/NNX/Source/NNXCore/Public/NNXTypes.h
@@ -31,32 +31,35 @@ enum class EMLTensorDataType : uint8
 
 namespace NNX
 {
-
-struct FSymbolicTensorShape
+	
+template <typename DimType, class FinalType> struct TTensorShapeBase
 {
 	constexpr static int32	MaxRank = 8;
-	TArray<int32, TInlineAllocator<MaxRank>> Data;
+	TArray<DimType, TInlineAllocator<MaxRank>> Data;
 
 	int32 Num() const
 	{
 		return Data.Num();
 	}
 
-	static FSymbolicTensorShape Make(TConstArrayView<const int32> Data)
+	static FinalType Make(TConstArrayView<const DimType> Data)
 	{
 		if (Data.Num() > MaxRank)
 		{
-			//UE_LOG(LogNNX, Warning, TEXT("Cannot create symbolic tensor shape, input is rank %d while max rank is %d"), Data.Num(), MaxRank);
+			//UE_LOG(LogNNX, Warning, TEXT("Cannot create tensor shape, input is rank %d while max rank is %d"), Data.Num(), MaxRank);
 			return {};
 		}
 
-		FSymbolicTensorShape Shape;
-		
+		FinalType Shape;
 		Shape.Data.Append(Data);
 		return Shape;
 	}
-	
-	bool IsConcrete()
+};
+
+
+struct FSymbolicTensorShape : TTensorShapeBase<int32, FSymbolicTensorShape>
+{
+	bool IsConcrete() const
 	{
 		for (int32 i = 0; i < Data.Num(); ++i)
 		{
@@ -69,18 +72,39 @@ struct FSymbolicTensorShape
 	}
 };
 
-struct FConcreteTensorShape
+struct FTensorShape : TTensorShapeBase<uint32, FTensorShape>
 {
-	TArray<uint32, TInlineAllocator<FSymbolicTensorShape::MaxRank>> Data;
+	uint64 Volume() const
+	{
+		uint64 Result = 1;
 
-	int32 Num() const
+		for (int32 Idx = 0; Idx < Data.Num(); ++Idx)
+		{
+			Result *= Data[Idx];
+		}
+		
+		return Result;
+	}
+	
+	bool IsCompatibleWith(FSymbolicTensorShape SymbolicShape) const
 	{
-		return Data.Num();
+		if (Num() != SymbolicShape.Num())
+		{
+			return false;
+		}
+		for (int32 i = 0; i < Num(); ++i)
+		{
+			if (SymbolicShape.Data[i] >= 0 && SymbolicShape.Data[i] != Data[i])
+			{
+				return false;
+			}
+		}
+		return true;
 	}
 	
-	static FConcreteTensorShape Make(const FSymbolicTensorShape& SymbolicShape)
+	static FTensorShape MakeFromSymbolic(const FSymbolicTensorShape& SymbolicShape)
 	{
-		FConcreteTensorShape ConcreteShape;
+		FTensorShape ConcreteShape;
 		for (int32 i = 0; i < SymbolicShape.Data.Num(); ++i)
 		{
 			int32 Dim = SymbolicShape.Data[i];
@@ -88,51 +112,6 @@ struct FConcreteTensorShape
 		}
 		return ConcreteShape;
 	}
-
-	static FConcreteTensorShape Make(TConstArrayView<const uint32> Data)
-	{
-		if (Data.Num() > FSymbolicTensorShape::MaxRank)
-		{
-			//UE_LOG(LogNNX, Warning, TEXT("Cannot create concrete tensor shape, input is rank %d while max rank is %d"), Data.Num(), MaxRank);
-			return {};
-		}
-
-		FConcreteTensorShape Shape;
-
-		Shape.Data.Append(Data);
-		return Shape;
-	}
-};
-
-/**
- *  Tensor descriptor
- */
-struct FMLTensorDesc
-{
-	constexpr static int32	MaxTensorDimension = 5;
-
-	FString					Name;
-	FConcreteTensorShape	Shape;
-	uint64					DataSize;//!< Size of data in bytes
-	EMLTensorDataType		DataType;
-
-	/** Make tensor descriptor */
-	static FMLTensorDesc Make(const FString& Name, const FConcreteTensorShape& Shape, EMLTensorDataType DataType, uint64 DataSize = 0)
-	{
-		return {Name, Shape, DataSize, DataType};
-	}
-
-	/** Check if descriptor is valid */
-	bool IsValid() const;
-
-	/** Return size of one element in bytes */
-	int32 GetElemByteSize() const;
-
-	/** Return Volume, i.e. number of elements */
-	int32 Volume() const;
-
-	/** Return number of elements, same as Volume() */
-	int32 Num() const;
 };
 
 /** Return data size in bytes for tensor data type */
@@ -140,67 +119,91 @@ inline size_t GetTensorDataTypeSizeInBytes(EMLTensorDataType InType)
 {
 	switch (InType)
 	{
-		case EMLTensorDataType::Complex128:
-			return 16;
-
-		case EMLTensorDataType::Complex64:
-			return 8;
-
-		case EMLTensorDataType::Double:
-		case EMLTensorDataType::Int64:
-		case EMLTensorDataType::UInt64:
-			return 8;
-
-		case EMLTensorDataType::Float:
-		case EMLTensorDataType::Int32:
-		case EMLTensorDataType::UInt32:
-			return 4;
-
-		case EMLTensorDataType::Half:
-		case EMLTensorDataType::BFloat16:
-		case EMLTensorDataType::Int16:
-		case EMLTensorDataType::UInt16:
-			return 2;
-
-		case EMLTensorDataType::Int8:
-		case EMLTensorDataType::UInt8:
-		case EMLTensorDataType::Char:
-		case EMLTensorDataType::Boolean:
-			return 1;
+	case EMLTensorDataType::Complex128:
+		return 16;
+
+	case EMLTensorDataType::Complex64:
+		return 8;
+
+	case EMLTensorDataType::Double:
+	case EMLTensorDataType::Int64:
+	case EMLTensorDataType::UInt64:
+		return 8;
+
+	case EMLTensorDataType::Float:
+	case EMLTensorDataType::Int32:
+	case EMLTensorDataType::UInt32:
+		return 4;
+
+	case EMLTensorDataType::Half:
+	case EMLTensorDataType::BFloat16:
+	case EMLTensorDataType::Int16:
+	case EMLTensorDataType::UInt16:
+		return 2;
+
+	case EMLTensorDataType::Int8:
+	case EMLTensorDataType::UInt8:
+	case EMLTensorDataType::Char:
+	case EMLTensorDataType::Boolean:
+		return 1;
 	}
 
 	return 0;
 }
 
-/** Check if descriptor is valid */
-inline bool FMLTensorDesc::IsValid() const
+struct FTensorDescBase
 {
-	return DataType != EMLTensorDataType::None && Shape.Data.Num() > 0 && Shape.Data.Num() <= MaxTensorDimension; //&& Volume() >= 1;
-}
+	FString					Name;
+	EMLTensorDataType		DataType;
 
-/** Return size of one element in bytes */
-inline int32 FMLTensorDesc::GetElemByteSize() const
-{
-	return GetTensorDataTypeSizeInBytes(DataType);
-}
+	/** Return size of one element in bytes */
+	uint32 GetElemByteSize() const
+	{
+		return GetTensorDataTypeSizeInBytes(DataType);
+	}
+	
+	bool IsValid() const
+	{
+		return DataType != EMLTensorDataType::None;
+	}
+};
 
-/** Return Volume, i.e. number of elements */
-inline int32 FMLTensorDesc::Volume() const
+/** Symbolic tensor descriptor */
+struct FSymbolicTensorDesc : FTensorDescBase
 {
-	int32 Result = 1;
+	FSymbolicTensorShape	Shape;
 
-	for (int32 Idx = 0; Idx < Shape.Data.Num(); ++Idx)
+	static FSymbolicTensorDesc Make(const FString& Name, const FSymbolicTensorShape& Shape, EMLTensorDataType DataType)
 	{
-		Result *= Shape.Data[Idx];
+		return { {Name, DataType} , Shape };
 	}
 
-	return Result;
-}
+	bool IsConcrete() const
+	{
+		return Shape.IsConcrete();
+	}
+};
 
-/** Return number of elements */
-inline int32 FMLTensorDesc::Num() const
+/** Concrete tensor descriptor */
+ struct FMLTensorDesc : FTensorDescBase
 {
-	return Volume();
-}
+	FTensorShape Shape;
+	uint32 Volume;
+	uint64 DataSize;
+
+	static FMLTensorDesc Make(const FString& Name, const FTensorShape& Shape, EMLTensorDataType DataType)
+	{
+		uint64 Volume = Shape.Volume();
+		uint64 DataSize = (uint64)GetTensorDataTypeSizeInBytes(DataType) * Volume;
+		
+		check(Volume <= TNumericLimits<uint32>::Max());
+		return { {Name, DataType}, Shape, (uint32)Volume, DataSize };
+	}
+
+	static FMLTensorDesc MakeFromSymbolic(const FSymbolicTensorDesc& TensorDesc)
+	{
+		return Make(TensorDesc.Name, FTensorShape::MakeFromSymbolic(TensorDesc.Shape), TensorDesc.DataType);
+	}
+};
 
 } // namespace NNX
diff --git a/Engine/Plugins/Experimental/NNX/Source/NNXQA/Private/NNXQAParametricTest.cpp b/Engine/Plugins/Experimental/NNX/Source/NNXQA/Private/NNXQAParametricTest.cpp
--- a/Engine/Plugins/Experimental/NNX/Source/NNXQA/Private/NNXQAParametricTest.cpp
+++ b/Engine/Plugins/Experimental/NNX/Source/NNXQA/Private/NNXQAParametricTest.cpp
@@ -118,11 +118,11 @@ namespace Test
 			TestSetup.Tags = TestTarget.Tags;
 		}
 
-		static FConcreteTensorShape GetConcreteShapeFromJsonArray(const TArray<int32>& JsonShape)
+		static FTensorShape GetConcreteShapeFromJsonArray(const TArray<int32>& JsonShape)
 		{
 			FSymbolicTensorShape SymbolicShape = FSymbolicTensorShape::Make(JsonShape);
 			check(SymbolicShape.IsConcrete());
-			FConcreteTensorShape ConcreteShape = FConcreteTensorShape::Make(SymbolicShape);
+			FTensorShape ConcreteShape = FTensorShape::MakeFromSymbolic(SymbolicShape);
 			return ConcreteShape;
 		}
 
@@ -145,7 +145,7 @@ namespace Test
 			uint32 i = 0;
 			for (auto&& Tensor : TestDataset.Inputs)
 			{
-				FConcreteTensorShape Shape = GetConcreteShapeFromJsonArray(Tensor.Shape);
+				FTensorShape Shape = GetConcreteShapeFromJsonArray(Tensor.Shape);
 				EMLTensorDataType TensorType = GetTensorTypeFromJson(Tensor.Type, DefaultInputType);
 				TestSetup.Inputs.Emplace(FMLTensorDesc::Make(FString::Printf(TEXT("in%d"), i++), Shape, TensorType));
 			}
@@ -153,14 +153,14 @@ namespace Test
 			i = 0;
 			for (auto&& Tensor : TestDataset.Outputs)
 			{
-				FConcreteTensorShape Shape = GetConcreteShapeFromJsonArray(Tensor.Shape);
+				FTensorShape Shape = GetConcreteShapeFromJsonArray(Tensor.Shape);
 				EMLTensorDataType TensorType = GetTensorTypeFromJson(Tensor.Type, DefaultOutputType);
 				TestSetup.Outputs.Emplace(FMLTensorDesc::Make(FString::Printf(TEXT("output%d"), i++), Shape, TensorType));
 			}
 			//If output is not defined it is the first input shape.
 			if (TestDataset.Outputs.Num() == 0 && TestDataset.Inputs.Num() > 0)
 			{
-				FConcreteTensorShape Shape = GetConcreteShapeFromJsonArray(TestDataset.Inputs[0].Shape);
+				FTensorShape Shape = GetConcreteShapeFromJsonArray(TestDataset.Inputs[0].Shape);
 				EMLTensorDataType TensorType = GetTensorTypeFromJson(TestDataset.Inputs[0].Type, DefaultOutputType);
 				TestSetup.Outputs.Emplace(FMLTensorDesc::Make(TEXT("output"), Shape, TensorType));
 			}
diff --git a/Engine/Plugins/Experimental/NNX/Source/NNXQA/Private/NNXQAUtils.cpp b/Engine/Plugins/Experimental/NNX/Source/NNXQA/Private/NNXQAUtils.cpp
--- a/Engine/Plugins/Experimental/NNX/Source/NNXQA/Private/NNXQAUtils.cpp
+++ b/Engine/Plugins/Experimental/NNX/Source/NNXQA/Private/NNXQAUtils.cpp
@@ -14,20 +14,20 @@ namespace NNX
 {
 namespace Test 
 {
-	static void FillTensorBindings(const TArrayView<const FMLTensorDesc>& TensorsDesc,
+	static void FillTensorBindings(TConstArrayView<const FMLTensorDesc> TensorDescs,
 		TArray<TArray<char>>& OutMemBuffers, TArray<FMLTensorBinding>& OutBindings,
         std::function<float(EMLTensorDataType, uint32, uint32)> Initializer)
 	{
 		OutMemBuffers.Empty();
 		OutBindings.Empty();
 
-		for (int Index = 0; Index < TensorsDesc.Num(); ++Index)
+		for (int Index = 0; Index < TensorDescs.Num(); ++Index)
 		{
-			const FMLTensorDesc& CurrentInputDesc = TensorsDesc[Index];
-			const uint64 NumberOfElements = CurrentInputDesc.Volume();
-			const int32 ElementByteSize = CurrentInputDesc.GetElemByteSize();
-			const uint64 BufferSize = ElementByteSize * NumberOfElements;
-			const EMLTensorDataType DataType = CurrentInputDesc.DataType;
+			const FMLTensorDesc& TensorDesc = TensorDescs[Index];
+			const uint32 NumberOfElements = TensorDesc.Volume;
+			const int32 ElementByteSize = TensorDesc.GetElemByteSize();
+			const uint64 BufferSize = TensorDesc.DataSize;
+			const EMLTensorDataType DataType = TensorDesc.DataType;
 
 			TArray<char> CurInputBuffer;
 	        CurInputBuffer.SetNum(BufferSize);
@@ -65,8 +65,8 @@ namespace Test
 			);
 		}
 
-		check(TensorsDesc.Num() == OutMemBuffers.Num());
-		check(TensorsDesc.Num() == OutBindings.Num());
+		check(TensorDescs.Num() == OutMemBuffers.Num());
+		check(TensorDescs.Num() == OutBindings.Num());
 	}
 
 	template<typename T> FString ShapeToString(TArrayView<const T> Shape)
@@ -94,7 +94,7 @@ namespace Test
 		TensorDesc += FString::Printf(TEXT("Name: %s, Shape: "), *desc.Name);
 		TensorDesc += ShapeToString(Shape);
 		const FString& DataTypeName = StaticEnum<EMLTensorDataType>()->GetNameStringByValue(static_cast<int64>(desc.DataType));
-		TensorDesc += FString::Printf(TEXT(" DataSize: %d, DataType: %s"), desc.DataSize, *DataTypeName);
+		TensorDesc += FString::Printf(TEXT(" DataType: %s"), *DataTypeName);
 
 		return TensorDesc;
 	}
@@ -107,14 +107,14 @@ namespace Test
 		TensorLog += FMLTensorDescToString(TensorDesc);
 		TensorLog += TEXT(", Data: ");
 
-		const constexpr int32 MAX_DATA_TO_LOG = 10;
-		const int32 MaxIndex = FMath::Min(MAX_DATA_TO_LOG, TensorDesc.Volume());
-        const int32 ElementByteSize = TensorDesc.GetElemByteSize();
-		for (int32 i = 0; i < MaxIndex; ++i)
+		const constexpr uint32 MAX_DATA_TO_LOG = 10;
+		const uint32 MaxIndex = FMath::Min(MAX_DATA_TO_LOG, TensorDesc.Volume);
+        const uint32 ElementByteSize = TensorDesc.GetElemByteSize();
+		for (uint32 i = 0; i < MaxIndex; ++i)
 		{
-			const int32 ByteOffset = i * ElementByteSize;
+			const uint32 ByteOffset = i * ElementByteSize;
 			const char* Data = TensorData.GetData() + ByteOffset;
-			check(ByteOffset <= TensorData.Num());
+			check((int32)ByteOffset <= TensorData.Num());
 
 			switch (TensorDesc.DataType)
 			{
@@ -133,7 +133,7 @@ namespace Test
 			if (i < MaxIndex)
 				TensorLog += TEXT(", ");
 		}
-		if (MaxIndex < TensorDesc.Volume())
+		if (MaxIndex < TensorDesc.Volume)
 			TensorLog += TEXT(",...");
 
 		return TensorLog;
@@ -147,10 +147,10 @@ namespace Test
 		const T* RefBuffer = (T*)RefRawBuffer.GetData();
 		const T* OtherBuffer = (T*)OtherRawBuffer.GetData();
 		
-		const int32 Volume = RefTensorDesc.Volume();
-		const int32 ElementByteSize = RefTensorDesc.GetElemByteSize();
+		const uint32 Volume = RefTensorDesc.Volume;
+		const uint32 ElementByteSize = RefTensorDesc.GetElemByteSize();
 
-		check(Volume == OtherTensorDesc.Volume());
+		check(Volume == OtherTensorDesc.Volume);
 		check(Volume * ElementByteSize == RefRawBuffer.Num());
 		check(Volume * ElementByteSize == OtherRawBuffer.Num());
 
@@ -171,7 +171,7 @@ namespace Test
 		int32 NumMissingNaNsInResults = 0;
 		int32 FirstMissingNaNIndex = -1;
 
-		for (int32 i = 0; i < Volume; ++i)
+		for (uint32 i = 0; i < Volume; ++i)
 		{
 			//All type are converted to float for comparison purpose
 			const float Result = (float)OtherBuffer[i];
@@ -249,9 +249,8 @@ namespace Test
 		bool bTensorDescMatch = true;
 		bTensorDescMatch &= (RefTensorDesc.Name == OtherTensorDesc.Name);
 		bTensorDescMatch &= (RefTensorDesc.Shape.Num() == OtherTensorDesc.Shape.Num());
-		bTensorDescMatch &= (RefTensorDesc.DataSize == OtherTensorDesc.DataSize);
 		bTensorDescMatch &= (RefTensorDesc.DataType == OtherTensorDesc.DataType);
-		check(RefTensorDesc.Shape.Num() <= RefTensorDesc.MaxTensorDimension);
+		check(RefTensorDesc.Shape.Num() <= FTensorShape::MaxRank);
 		for (int32 i = 0; i < RefTensorDesc.Shape.Num(); ++i)
 		{
 			bTensorDescMatch &= (RefTensorDesc.Shape.Data[i] == OtherTensorDesc.Shape.Data[i]);
@@ -326,8 +325,8 @@ namespace Test
 		return FMath::Sin((float)(ElementIndex + IndexOffsetBetweenTensor * TensorIndex));
 	}
 
-	static int RunTestInference(const FNNIModelRaw& ONNXModelData, NNX::IRuntime* Runtime, 
-		TArray<NNX::FMLTensorDesc>& OutOutputTensorsDesc, TArray<TArray<char>>& OutOutputMemBuffers)
+	static int RunTestInference(const FNNIModelRaw& ONNXModelData, IRuntime* Runtime, 
+		TArray<FMLTensorDesc>& OutOutputTensorsDesc, TArray<TArray<char>>& OutOutputMemBuffers)
 	{
 		OutOutputMemBuffers.Empty();
 		OutOutputTensorsDesc.Empty();
@@ -343,7 +342,7 @@ namespace Test
 		}
 		
 		UMLInferenceModel* UInferenceModel = UMLInferenceModel::CreateFromFormatDesc(RuntimeModelData);
-		TUniquePtr<NNX::FMLInferenceModel> InferenceModel(Runtime->CreateInferenceModel(UInferenceModel));
+		TUniquePtr<FMLInferenceModel> InferenceModel(Runtime->CreateInferenceModel(UInferenceModel));
 		
 		if (!InferenceModel.IsValid())
 		{
@@ -351,16 +350,33 @@ namespace Test
 			return -1;
 		}
 
-		//bind tensors to memory (CPU) and initialize
-		const TArrayView<const NNX::FMLTensorDesc> InputTensorsDescView = InferenceModel->GetInputTensors();
-		TArray<NNX::FMLTensorBinding> InputBindings;
+		//bind input tensors to memory (CPU) and initialize
+		TConstArrayView<const FSymbolicTensorDesc> InputSymbolicTensorDescs = InferenceModel->GetInputTensorDescs();
+		TArray<FTensorShape> InputShapes;
+		TArray<FMLTensorDesc> InputTensorsDescs;
+		TArray<FMLTensorBinding> InputBindings;
 		TArray<TArray<char>> InputMemBuffers;
-		FillTensorBindings(InputTensorsDescView, InputMemBuffers, InputBindings, InputTensorInitializer);
-
-		const TArrayView<const NNX::FMLTensorDesc> OutputTensorsDescView = InferenceModel->GetOutputTensors();
+		
+		for (const FSymbolicTensorDesc& SymbolicTensorDesc : InputSymbolicTensorDescs)
+		{
+			FMLTensorDesc TensorDesc = FMLTensorDesc::MakeFromSymbolic(SymbolicTensorDesc);
+			InputShapes.Emplace(TensorDesc.Shape);
+			InputTensorsDescs.Emplace(TensorDesc);
+		}
+		FillTensorBindings(InputTensorsDescs, InputMemBuffers, InputBindings, InputTensorInitializer);
+		
+		//bind output tensors to memory (CPU) and initialize
+		TConstArrayView<const FSymbolicTensorDesc> OutputSymbolicTensorDescs = InferenceModel->GetOutputTensorDescs();
+		TArray<FMLTensorDesc> OutputTensorsDescs;
 		TArray<NNX::FMLTensorBinding> OutputBindings;
-		OutOutputTensorsDesc = OutputTensorsDescView;
-		FillTensorBindings(OutputTensorsDescView, OutOutputMemBuffers, OutputBindings, OutputTensorInitializer);
+
+		for (const FSymbolicTensorDesc& SymbolicTensorDesc : OutputSymbolicTensorDescs)
+		{
+			FMLTensorDesc TensorDesc = FMLTensorDesc::MakeFromSymbolic(SymbolicTensorDesc);
+			OutputTensorsDescs.Emplace(TensorDesc);
+		}
+		OutOutputTensorsDesc = OutputTensorsDescs;
+		FillTensorBindings(OutputTensorsDescs, OutOutputMemBuffers, OutputBindings, OutputTensorInitializer);
 
 		check(OutOutputTensorsDesc.Num() == OutOutputMemBuffers.Num());
 
@@ -382,7 +398,7 @@ namespace Test
 		#endif //UE_BUILD_DEBUG
 
 		//run inference
-		int ReturnValue = InferenceModel->Run(InputBindings, OutputBindings);
+		int ReturnValue = InferenceModel->Run(InputBindings, InputShapes, OutputBindings);
 		return ReturnValue;
 	}
 
diff --git a/Engine/Plugins/Experimental/NNX/Source/NNXQA/Resources/NNXQATestDesc.json b/Engine/Plugins/Experimental/NNX/Source/NNXQA/Resources/NNXQATestDesc.json
--- a/Engine/Plugins/Experimental/NNX/Source/NNXQA/Resources/NNXQATestDesc.json
+++ b/Engine/Plugins/Experimental/NNX/Source/NNXQA/Resources/NNXQATestDesc.json
@@ -76,10 +76,7 @@
 				},
 				{
 					"target": "TwoLayersNoWeight",
-					"runtimes": [
-						{ "name": "NNXRuntimeHlsl",	"skip": true, "comment": "Model with multiple layer not yet implemented."	},
-						{ "name": "NNXRuntimeDml",  "skip": true, "comment": "Model with multiple layer not yet implemented." }
-					]
+					"tags": [ "SmokeTest" ]
 				}
 			]
 		}
diff --git a/Engine/Plugins/Experimental/NNX/Source/NNXRuntimeCPU/Private/NNXRuntimeCPU.cpp b/Engine/Plugins/Experimental/NNX/Source/NNXRuntimeCPU/Private/NNXRuntimeCPU.cpp
--- a/Engine/Plugins/Experimental/NNX/Source/NNXRuntimeCPU/Private/NNXRuntimeCPU.cpp
+++ b/Engine/Plugins/Experimental/NNX/Source/NNXRuntimeCPU/Private/NNXRuntimeCPU.cpp
@@ -163,19 +163,18 @@ bool FMLInferenceModelCPU::ConfigureTensors(const bool InIsInput)
 
 	const uint32 NumberTensors = bIsInput ? Session->GetInputCount() : Session->GetOutputCount();
 
-	TArray<NNX::FMLTensorDesc>& TensorsDescriptors = bIsInput ? InputTensors : OutputTensors;
+	TArray<FSymbolicTensorDesc>& SymbolicTensorDescs = bIsInput ? InputSymbolicTensors : OutputSymbolicTensors;
 	TArray<ONNXTensorElementDataType>& TensorsORTType = bIsInput ? InputTensorsORTType : OutputTensorsORTType;
 	TArray<const char*>& TensorNames = bIsInput ? InputTensorNames : OutputTensorNames;
 
 	for (uint32 TensorIndex = 0; TensorIndex < NumberTensors; ++TensorIndex)
 	{
-		NNX::FMLTensorDesc CurrentTensorDescriptor;
+		FSymbolicTensorDesc SymbolicTensorDesc;
 
 		// Get Tensor name
 		const char* CurTensorName = bIsInput ? Session->GetInputName(TensorIndex, *Allocator) : Session->GetOutputName(TensorIndex, *Allocator);
 		TensorNames.Emplace(CurTensorName);
-		CurrentTensorDescriptor.Name = FString(CurTensorName);
-		
+		SymbolicTensorDesc.Name = FString(CurTensorName);
 
 		// Get node type
 		Ort::TypeInfo CurrentTypeInfo = bIsInput ? Session->GetInputTypeInfo(TensorIndex) : Session->GetOutputTypeInfo(TensorIndex);
@@ -186,38 +185,55 @@ bool FMLInferenceModelCPU::ConfigureTensors(const bool InIsInput)
 		TensorsORTType.Emplace(ONNXTensorElementDataTypeEnum);
 
 		std::pair<EMLTensorDataType, uint64> TypeAndSize = TranslateTensorTypeORTToNNI(ONNXTensorElementDataTypeEnum);
-		CurrentTensorDescriptor.DataType = TypeAndSize.first;
+		SymbolicTensorDesc.DataType = TypeAndSize.first;
 
-		CurrentTensorDescriptor.Shape.Data.SetNumUninitialized(CurrentTensorInfo.GetShape().size());
-		uint8 Index = 0;
-		uint64 TensorNumberElements = 1;
+		SymbolicTensorDesc.Shape.Data.Reserve(CurrentTensorInfo.GetShape().size());
 		for (const int64_t CurrentTensorSize : CurrentTensorInfo.GetShape())
 		{
-			// Negative (variable) dimensions not implemented yet
-			if (CurrentTensorSize < 0)
+			SymbolicTensorDesc.Shape.Data.Add((int32)CurrentTensorSize);
+		}
+
+		if (!SymbolicTensorDesc.IsConcrete())
 			{
-				CurrentTensorDescriptor.Shape.Data[Index] = 1;
 				UE_LOG(LogNNX, Display,
 					TEXT("Negative (i.e., variable) dimensions not allowed yet, hard-coded to 1. Let us know if you really need variable dimensions."
 						" Keep in mind that fixed sizes might allow additional optimizations and speedup of the network during Run()."));
 			}
-			else
+
+		check(SymbolicTensorDesc.GetElemByteSize() == TypeAndSize.second);
+		SymbolicTensorDescs.Emplace(SymbolicTensorDesc);
+	}
+
+	return true;
+}
+
+int FMLInferenceModelCPU::SetInputShapes(TConstArrayView<const FTensorShape> InInputShapes)
 			{
-				CurrentTensorDescriptor.Shape.Data[Index] = CurrentTensorSize;
-				TensorNumberElements *= CurrentTensorSize;
-			}
-			++Index;
+	//Verify input shape are valid for the model
+	if (FMLInferenceModel::SetInputShapes(InInputShapes) != 0)
+	{
+		return -1;
 		}
 
-		CurrentTensorDescriptor.DataSize = TensorNumberElements * TypeAndSize.second;
-		TensorsDescriptors.Emplace(CurrentTensorDescriptor);
+	//Run shape inference
+	//TODO jira 168972: handle dynamic input shape
 
+	//Build concrete input and output tensor descriptions
+	InputTensors.Empty();
+	for (FSymbolicTensorDesc SymbolicTensorDesc : InputSymbolicTensors)
+	{
+		InputTensors.Emplace(FMLTensorDesc::MakeFromSymbolic(SymbolicTensorDesc));
 	}
-	return true;
+	OutputTensors.Empty();
+	for (FSymbolicTensorDesc SymbolicTensorDesc : OutputSymbolicTensors)
+	{
+		OutputTensors.Emplace(FMLTensorDesc::MakeFromSymbolic(SymbolicTensorDesc));
 }
 
+	return 0;
+}
 
-int FMLInferenceModelCPU::Run(TArrayView<const NNX::FMLTensorBinding> InInputBindingTensors, TArrayView<const NNX::FMLTensorBinding> OutOutputBindingTensors)
+int FMLInferenceModelCPU::Run(TConstArrayView<const FMLTensorBinding> InInputBindings, TConstArrayView<const FTensorShape> InInputShapes, TConstArrayView<const FMLTensorBinding> InOutputBindings)
 {
 	DECLARE_SCOPE_CYCLE_COUNTER(TEXT("FMLInferenceModelCPU_Run"), STAT_FMLInferenceModelCPU_Run, STATGROUP_MachineLearning);
 
@@ -228,6 +244,14 @@ int FMLInferenceModelCPU::Run(TArrayView<const NNX::FMLTensorBinding> InInputBin
 		return -1;
 	}
 
+	// Prepare the model for the concrete input shapes
+	int Res = SetInputShapes(InInputShapes);
+	if (Res != 0)
+	{
+		UE_LOG(LogNNX, Warning, TEXT("Model preparation failed with error code:%d"), Res);
+		return -1;
+	}
+
 	UE::NNEProfiling::Internal::FTimer RunTimer;
 	RunTimer.Tic();
 
@@ -241,10 +265,10 @@ int FMLInferenceModelCPU::Run(TArrayView<const NNX::FMLTensorBinding> InInputBin
 #endif //WITH_EDITOR
 	{
 		TArray<Ort::Value> InputOrtTensors;
-		BindTensorsToORT(InInputBindingTensors, InputTensors, InputTensorsORTType, AllocatorInfo.Get(), InputOrtTensors);
+		BindTensorsToORT(InInputBindings, InputTensors, InputTensorsORTType, AllocatorInfo.Get(), InputOrtTensors);
 
 		TArray<Ort::Value> OutputOrtTensors;
-		BindTensorsToORT(OutOutputBindingTensors, OutputTensors, OutputTensorsORTType, AllocatorInfo.Get(), OutputOrtTensors);
+		BindTensorsToORT(InOutputBindings, OutputTensors, OutputTensorsORTType, AllocatorInfo.Get(), OutputOrtTensors);
 
 		Session->Run(Ort::RunOptions{ nullptr },
 			InputTensorNames.GetData(), &InputOrtTensors[0], InputTensorNames.Num(),
diff --git a/Engine/Plugins/Experimental/NNX/Source/NNXRuntimeCPU/Public/NNXRuntimeCPU.h b/Engine/Plugins/Experimental/NNX/Source/NNXRuntimeCPU/Public/NNXRuntimeCPU.h
--- a/Engine/Plugins/Experimental/NNX/Source/NNXRuntimeCPU/Public/NNXRuntimeCPU.h
+++ b/Engine/Plugins/Experimental/NNX/Source/NNXRuntimeCPU/Public/NNXRuntimeCPU.h
@@ -84,7 +84,8 @@ namespace NNX
 		bool Init(const UMLInferenceModel* InferenceModel);
 		bool IsLoaded() const;
 
-		virtual int Run(TArrayView<const NNX::FMLTensorBinding> InInputBindingTensors, TArrayView<const NNX::FMLTensorBinding> OutOutputBindingTensors);
+		virtual int SetInputShapes(TConstArrayView<const FTensorShape> InInputShapes) override;
+		virtual int Run(TConstArrayView<const FMLTensorBinding> InInputBindings, TConstArrayView<const FTensorShape> InInputShapes, TConstArrayView<const FMLTensorBinding> InOutputBindings) override;
 
 		float GetLastRunTimeMSec() const;
 		UE::NNEProfiling::Internal::FStatistics GetRunStatistics() const;
@@ -113,6 +114,9 @@ namespace NNX
 		TArray<const char*> InputTensorNames;
 		TArray<const char*> OutputTensorNames;
 
+		TArray<FMLTensorDesc> InputTensors;
+		TArray<FMLTensorDesc> OutputTensors;
+
 		bool InitializedAndConfigureMembers();
 		bool ConfigureTensors(const bool InIsInput);
 
diff --git a/Engine/Plugins/Experimental/NNX/Source/NNXRuntimeORT/Private/NNXRuntimeORT.cpp b/Engine/Plugins/Experimental/NNX/Source/NNXRuntimeORT/Private/NNXRuntimeORT.cpp
--- a/Engine/Plugins/Experimental/NNX/Source/NNXRuntimeORT/Private/NNXRuntimeORT.cpp
+++ b/Engine/Plugins/Experimental/NNX/Source/NNXRuntimeORT/Private/NNXRuntimeORT.cpp
@@ -237,19 +237,18 @@ bool FMLInferenceModelORT::ConfigureTensors(const bool InIsInput)
 
 	const uint32 NumberTensors = bIsInput ? Session->GetInputCount() : Session->GetOutputCount();
 
-	TArray<NNX::FMLTensorDesc>& TensorsDescriptors = bIsInput ? InputTensors : OutputTensors;
+	TArray<FSymbolicTensorDesc>& SymbolicTensorDescs = bIsInput ? InputSymbolicTensors : OutputSymbolicTensors;
 	TArray<ONNXTensorElementDataType>& TensorsORTType = bIsInput ? InputTensorsORTType : OutputTensorsORTType;
 	TArray<const char*>& TensorNames = bIsInput ? InputTensorNames : OutputTensorNames;
 
-
 	for (uint32 TensorIndex = 0; TensorIndex < NumberTensors; ++TensorIndex)
 	{
-		NNX::FMLTensorDesc CurrentTensorDescriptor;
+		FSymbolicTensorDesc SymbolicTensorDesc;
 
 		// Get Tensor name
 		const char* CurTensorName = bIsInput ? Session->GetInputName(TensorIndex, *Allocator) : Session->GetOutputName(TensorIndex, *Allocator);
 		TensorNames.Emplace(CurTensorName);
-		CurrentTensorDescriptor.Name = FString(CurTensorName);
+		SymbolicTensorDesc.Name = FString(CurTensorName);
 
 		// Get node type
 		Ort::TypeInfo CurrentTypeInfo = bIsInput ? Session->GetInputTypeInfo(TensorIndex) : Session->GetOutputTypeInfo(TensorIndex);
@@ -260,37 +259,55 @@ bool FMLInferenceModelORT::ConfigureTensors(const bool InIsInput)
 		TensorsORTType.Emplace(ONNXTensorElementDataTypeEnum);
 
 		std::pair<EMLTensorDataType, uint64> TypeAndSize = TranslateTensorTypeORTToNNI(ONNXTensorElementDataTypeEnum);
-		CurrentTensorDescriptor.DataType = TypeAndSize.first;
+		SymbolicTensorDesc.DataType = TypeAndSize.first;
 		
-		CurrentTensorDescriptor.Shape.Data.SetNumUninitialized(CurrentTensorInfo.GetShape().size());
-		uint8 Index = 0;
-		uint64 TensorNumberElements = 1;
+		SymbolicTensorDesc.Shape.Data.Reserve(CurrentTensorInfo.GetShape().size());
 		for (const int64_t CurrentTensorSize : CurrentTensorInfo.GetShape())
 		{
-			// Negative (variable) dimensions not implemented yet
-			if (CurrentTensorSize < 0)
+			SymbolicTensorDesc.Shape.Data.Add((int32)CurrentTensorSize);
+		}
+
+		if (!SymbolicTensorDesc.IsConcrete())
 			{
-				CurrentTensorDescriptor.Shape.Data[Index] = 1;
 				UE_LOG(LogNNX, Display,
 					TEXT("Negative (i.e., variable) dimensions not allowed yet, hard-coded to 1. Let us know if you really need variable dimensions."
 						" Keep in mind that fixed sizes might allow additional optimizations and speedup of the network during Run()."));
 			}
-			else
+		
+		check(SymbolicTensorDesc.GetElemByteSize() == TypeAndSize.second);
+		SymbolicTensorDescs.Emplace(SymbolicTensorDesc);
+	}
+
+	return true;
+}
+
+int FMLInferenceModelORT::SetInputShapes(TConstArrayView<const FTensorShape> InInputShapes)
+{
+	//Verify input shape are valid for the model
+	if (FMLInferenceModel::SetInputShapes(InInputShapes) != 0)
 			{
-				CurrentTensorDescriptor.Shape.Data[Index] = CurrentTensorSize;
-				TensorNumberElements *= CurrentTensorSize;
+		return -1;
 			}
-			++Index;
+
+	//Run shape inference
+	//TODO jira 168972: handle dynamic input shape
+
+	//Build concrete input and output tensor descriptions
+	InputTensors.Empty();
+	for (FSymbolicTensorDesc SymbolicTensorDesc : InputSymbolicTensors)
+	{
+		InputTensors.Emplace(FMLTensorDesc::MakeFromSymbolic(SymbolicTensorDesc));
 		}
-		
-		CurrentTensorDescriptor.DataSize = TensorNumberElements * TypeAndSize.second;
-		TensorsDescriptors.Emplace(CurrentTensorDescriptor);
+	OutputTensors.Empty();
+	for (FSymbolicTensorDesc SymbolicTensorDesc : OutputSymbolicTensors)
+	{
+		OutputTensors.Emplace(FMLTensorDesc::MakeFromSymbolic(SymbolicTensorDesc));
 	}
 
-	return true;
+	return 0;
 }
 
-int FMLInferenceModelORT::Run(TArrayView<const NNX::FMLTensorBinding> InInputBindingTensors, TArrayView<const NNX::FMLTensorBinding> OutOutputBindingTensors)
+int FMLInferenceModelORT::Run(TConstArrayView<const FMLTensorBinding> InInputBindings, TConstArrayView<const FTensorShape> InInputShapes, TConstArrayView<const FMLTensorBinding> InOutputBindings)
 {
 	DECLARE_SCOPE_CYCLE_COUNTER(TEXT("FMLInferenceModelORT_Run"), STAT_FMLInferenceModelORT_Run, STATGROUP_MachineLearning);
 
@@ -301,6 +318,14 @@ int FMLInferenceModelORT::Run(TArrayView<const NNX::FMLTensorBinding> InInputBin
 		return -1;
 	}
 
+	// Prepare the model for the concrete input shapes
+	int Res = SetInputShapes(InInputShapes);
+	if (Res != 0)
+	{
+		UE_LOG(LogNNX, Warning, TEXT("Model preparation failed with error code:%d"), Res);
+		return -1;
+	}
+
 	UE::NNEProfiling::Internal::FTimer RunTimer;
 	RunTimer.Tic();
 
@@ -314,10 +339,10 @@ int FMLInferenceModelORT::Run(TArrayView<const NNX::FMLTensorBinding> InInputBin
 #endif //WITH_EDITOR
 	{
 		TArray<Ort::Value> InputOrtTensors;
-		BindTensorsToORT(InInputBindingTensors, InputTensors, InputTensorsORTType, AllocatorInfo.Get(), InputOrtTensors);
+		BindTensorsToORT(InInputBindings, InputTensors, InputTensorsORTType, AllocatorInfo.Get(), InputOrtTensors);
 
 		TArray<Ort::Value> OutputOrtTensors;
-		BindTensorsToORT(OutOutputBindingTensors, OutputTensors, OutputTensorsORTType, AllocatorInfo.Get(), OutputOrtTensors);
+		BindTensorsToORT(InOutputBindings, OutputTensors, OutputTensorsORTType, AllocatorInfo.Get(), OutputOrtTensors);
 
 		Session->Run(Ort::RunOptions{ nullptr },
 			InputTensorNames.GetData(), &InputOrtTensors[0], InputTensorNames.Num(),
diff --git a/Engine/Plugins/Experimental/NNX/Source/NNXRuntimeORT/Private/NNXRuntimeORT.h b/Engine/Plugins/Experimental/NNX/Source/NNXRuntimeORT/Private/NNXRuntimeORT.h
--- a/Engine/Plugins/Experimental/NNX/Source/NNXRuntimeORT/Private/NNXRuntimeORT.h
+++ b/Engine/Plugins/Experimental/NNX/Source/NNXRuntimeORT/Private/NNXRuntimeORT.h
@@ -183,7 +183,8 @@ namespace NNX
 		bool Init(const UMLInferenceModel* InferenceModel);
 		bool IsLoaded() const;
 
-		virtual int Run(TArrayView<const FMLTensorBinding> InInputBindingTensors, TArrayView<const FMLTensorBinding> OutOutputBindingTensors);
+		virtual int SetInputShapes(TConstArrayView<const FTensorShape> InInputShapes) override;
+		virtual int Run(TConstArrayView<const FMLTensorBinding> InInputBindings, TConstArrayView<const FTensorShape> InInputShapes, TConstArrayView<const FMLTensorBinding> InOutputBindings) override;
 
 		float GetLastRunTimeMSec() const;
 		UE::NNEProfiling::Internal::FStatistics GetRunStatistics() const;
@@ -211,6 +212,9 @@ namespace NNX
 		TArray<const char*> InputTensorNames;
 		TArray<const char*> OutputTensorNames;
 
+		TArray<FMLTensorDesc> InputTensors;
+		TArray<FMLTensorDesc> OutputTensors;
+
 		virtual bool InitializedAndConfigureMembers();
 		bool ConfigureTensors(const bool InIsInput);
 
diff --git a/Engine/Plugins/Experimental/NNX/Source/NNXRuntimeRDG/Private/Hlsl/NNIRuntimeRDGElementWiseBinary.cpp b/Engine/Plugins/Experimental/NNX/Source/NNXRuntimeRDG/Private/Hlsl/NNIRuntimeRDGElementWiseBinary.cpp
--- a/Engine/Plugins/Experimental/NNX/Source/NNXRuntimeRDG/Private/Hlsl/NNIRuntimeRDGElementWiseBinary.cpp
+++ b/Engine/Plugins/Experimental/NNX/Source/NNXRuntimeRDG/Private/Hlsl/NNIRuntimeRDGElementWiseBinary.cpp
@@ -49,7 +49,7 @@ namespace UE::NNIRuntimeRDG::Private::Hlsl
 			FRDGBufferSRVRef RHSInputSRV = GraphBuilder.CreateSRV(FRDGBufferSRVDesc(InInputBindings[1].Buffer, PF_R32_FLOAT));
 			FRDGBufferUAVRef OutputUAV = GraphBuilder.CreateUAV(FRDGBufferUAVDesc(OutOutputBindings[0].Buffer, PF_R32_FLOAT));
 
-			FIntVector ThreadGroupCount = NNX::ComputeElementWiseThreadGroups(Output.Num(), FElementWiseBinaryConstants::NUM_GROUP_THREADS);
+			FIntVector ThreadGroupCount = NNX::ComputeElementWiseThreadGroups(Output.Volume, FElementWiseBinaryConstants::NUM_GROUP_THREADS);
 
 			// Set parameters
 			TElementWiseBinaryCS::FParameters* Params = GraphBuilder.AllocParameters<TElementWiseBinaryCS::FParameters>();
@@ -59,7 +59,7 @@ namespace UE::NNIRuntimeRDG::Private::Hlsl
 			FillTensorStrideForBroadcastShaderParameters(LHSInput, Output.Shape.Num(), Params->TensorInfo, 0);
 			FillTensorStrideForBroadcastShaderParameters(RHSInput, Output.Shape.Num(), Params->TensorInfo, 1);
 			FillTensorStrideShaderParameters(Output, Params->TensorInfo, 2);
-			Params->Num = Output.Num();
+			Params->Num = Output.Volume;
 			Params->ThreadCountX = ThreadGroupCount.X * FElementWiseBinaryConstants::NUM_GROUP_THREADS;
 
 			TElementWiseBinaryCS::FPermutationDomain PermutationVector;
diff --git a/Engine/Plugins/Experimental/NNX/Source/NNXRuntimeRDG/Private/Hlsl/NNIRuntimeRDGElementWiseUnary.cpp b/Engine/Plugins/Experimental/NNX/Source/NNXRuntimeRDG/Private/Hlsl/NNIRuntimeRDGElementWiseUnary.cpp
--- a/Engine/Plugins/Experimental/NNX/Source/NNXRuntimeRDG/Private/Hlsl/NNIRuntimeRDGElementWiseUnary.cpp
+++ b/Engine/Plugins/Experimental/NNX/Source/NNXRuntimeRDG/Private/Hlsl/NNIRuntimeRDGElementWiseUnary.cpp
@@ -78,7 +78,7 @@ namespace UE::NNIRuntimeRDG::Private::Hlsl
 			FRDGBufferSRVRef InputSRV = GraphBuilder.CreateSRV(FRDGBufferSRVDesc(InInputBindings[0].Buffer, PF_R32_FLOAT));
 			FRDGBufferUAVRef OutputUAV = GraphBuilder.CreateUAV(FRDGBufferUAVDesc(OutOutputBindings[0].Buffer, PF_R32_FLOAT));
 		
-			FIntVector ThreadGroupCount = NNX::ComputeElementWiseThreadGroups(Output.Num(), FElementWiseUnaryConstants::NUM_GROUP_THREADS);
+			FIntVector ThreadGroupCount = NNX::ComputeElementWiseThreadGroups(Output.Volume, FElementWiseUnaryConstants::NUM_GROUP_THREADS);
 
 			// Set parameters
 			TElementWiseUnaryCS::FParameters* Params = GraphBuilder.AllocParameters<TElementWiseUnaryCS::FParameters>();
@@ -87,7 +87,7 @@ namespace UE::NNIRuntimeRDG::Private::Hlsl
 			Params->Alpha = Alpha;
 			Params->Beta = Beta;
 			Params->Gamma = Gamma;
-			Params->Num = Output.Num();
+			Params->Num = Output.Volume;
 			Params->ThreadCountX = ThreadGroupCount.X * FElementWiseUnaryConstants::NUM_GROUP_THREADS;
 
 			TElementWiseUnaryCS::FPermutationDomain PermutationVector;
diff --git a/Engine/Plugins/Experimental/NNX/Source/NNXRuntimeRDG/Private/Hlsl/NNIRuntimeRDGElementWiseVariadic.cpp b/Engine/Plugins/Experimental/NNX/Source/NNXRuntimeRDG/Private/Hlsl/NNIRuntimeRDGElementWiseVariadic.cpp
--- a/Engine/Plugins/Experimental/NNX/Source/NNXRuntimeRDG/Private/Hlsl/NNIRuntimeRDGElementWiseVariadic.cpp
+++ b/Engine/Plugins/Experimental/NNX/Source/NNXRuntimeRDG/Private/Hlsl/NNIRuntimeRDGElementWiseVariadic.cpp
@@ -35,7 +35,7 @@ namespace UE::NNIRuntimeRDG::Private::Hlsl
 		}
 
 		// Set parameters
-		FIntVector ThreadGroupCount = NNX::ComputeElementWiseThreadGroups(OutputDesc.Num(), FElementWiseVariadicConstants::NUM_GROUP_THREADS);
+		FIntVector ThreadGroupCount = NNX::ComputeElementWiseThreadGroups(OutputDesc.Volume, FElementWiseVariadicConstants::NUM_GROUP_THREADS);
 		TElementWiseVariadicCS::FParameters* Params = GraphBuilder.AllocParameters<TElementWiseVariadicCS::FParameters>();
 
 		Params->Input0 = InputsSRV[0];
@@ -57,7 +57,7 @@ namespace UE::NNIRuntimeRDG::Private::Hlsl
 			FillTensorStrideForBroadcastShaderParameters(InputDesc[3], OutputDesc.Shape.Num(), Params->InputTensorInfo, 3);
 		}
 		FillTensorStrideShaderParameters(OutputDesc, Params->OutputTensorInfo, 0);
-		Params->Num = OutputDesc.Num();
+		Params->Num = OutputDesc.Volume;
 		Params->ThreadCountX = ThreadGroupCount.X * FElementWiseVariadicConstants::NUM_GROUP_THREADS;
 		Params->Scale = Scale;
 
diff --git a/Engine/Plugins/Experimental/NNX/Source/NNXRuntimeRDG/Private/NNXRuntimeDML.cpp b/Engine/Plugins/Experimental/NNX/Source/NNXRuntimeRDG/Private/NNXRuntimeDML.cpp
--- a/Engine/Plugins/Experimental/NNX/Source/NNXRuntimeRDG/Private/NNXRuntimeDML.cpp
+++ b/Engine/Plugins/Experimental/NNX/Source/NNXRuntimeRDG/Private/NNXRuntimeDML.cpp
@@ -163,8 +163,8 @@ namespace DmlUtil
 	{
 		DML_BUFFER_TENSOR_DESC												BuffDesc;
 		DML_TENSOR_DESC														Desc;
-		TArray<uint32, TInlineAllocator<FMLTensorDesc::MaxTensorDimension>>	Sizes;
-		TArray<uint32, TInlineAllocator<FMLTensorDesc::MaxTensorDimension>>	Strides;
+		TArray<uint32, TInlineAllocator<FTensorShape::MaxRank>>		Sizes;
+		TArray<uint32, TInlineAllocator<FTensorShape::MaxRank>>		Strides;
 	};
 
 	void SetTensorStrides(FTensorDesc& TensorDesc, const FMLTensorDesc& InputDesc)
@@ -182,7 +182,7 @@ namespace DmlUtil
 
 	void SetTensorSizesAndStridesForBroadcast(FTensorDesc& TensorDesc, const FMLTensorDesc& InputDesc, const FMLTensorDesc& TargetDesc)
 	{
-		static_assert(FMLTensorDesc::MaxTensorDimension <= 8);
+		static_assert(FTensorShape::MaxRank <= 8);
 		
 		const uint32 TargetDimension = TargetDesc.Shape.Num() != -1 ? TargetDesc.Shape.Num() : InputDesc.Shape.Num();
 		checkf(TargetDesc.Shape.Num() >= InputDesc.Shape.Num(), TEXT("Can't broadcast tensor from rank %d to rank %d, should be inferior or equal."), InputDesc.Shape.Num(), TargetDimension);
@@ -552,7 +552,7 @@ public:
 	virtual bool Initialize(FDeviceContextDml* InDevCtx, TArrayView<const FMLTensorDesc> InputTensors, TArrayView<const FMLTensorDesc> OutputTensors, const UE::NNECore::FAttributeMap& Attributes) override
 	{
 		// TODO: Setup attributes
-		Num = InputTensors[0].Num();
+		Num = InputTensors[0].Volume;
 
 		DevCtx = InDevCtx;
 
@@ -802,7 +802,7 @@ public:
 	virtual bool Initialize(FDeviceContextDml* InDevCtx, TArrayView<const FMLTensorDesc> InputTensors, TArrayView<const FMLTensorDesc> OutputTensors, const UE::NNECore::FAttributeMap& Attributes) override
 	{
 		// TODO: Setup attributes
-		Num = OutputTensors[0].Num();
+		Num = OutputTensors[0].Volume;
 
 		DevCtx = InDevCtx;
 
@@ -980,7 +980,7 @@ public:
 
 protected:
 
-	virtual void AddDispatchOps_RenderThread(FRDGBuilder& GraphBuilder, TArrayView<const FMLTensorBinding> InInputBindings, TArrayView<const FMLTensorBinding> OutOutputBindings) override;
+	virtual void AddDispatchOps_RenderThread(FRDGBuilder& GraphBuilder) override;
 
 private:
 
@@ -1267,30 +1267,28 @@ bool FMLInferenceModelDml::Init(UMLInferenceModel* InModel, FDeviceContextDml* I
 		return false;
 	}
 
-	if (Format.Operators.Num() > 1)
-	{
-		UE_LOG(LogNNX, Warning, TEXT("Failed to create inference model, currently on single layer models are supported"));
-		return false;
-	}
-
 	DevCtx = InDevCtx;
 
 	// Loop over all operators in the model and create them
 	for (int32 Idx = 0; Idx < Format.Operators.Num(); ++Idx)
 	{
 		const FString TypeName = Format.Operators[Idx].TypeName;
-		
+
 		TArray<FMLTensorDesc> OpInputTensors;
 		TArray<FMLTensorDesc> OpOutputTensors;
 		UE::NNECore::FAttributeMap AttributeMap;
 
 		for (int32 InputTensorIndex : Format.Operators[Idx].InTensors)
 		{
-			OpInputTensors.Emplace(AllTensors[InputTensorIndex]);
+			FSymbolicTensorDesc SymbolicTensorDesc = AllSymbolicTensors[InputTensorIndex];
+			//TODO jira 168972: Handle dynamic tensor desc, op should init from symbolic shapes
+			OpInputTensors.Emplace(FMLTensorDesc::MakeFromSymbolic(SymbolicTensorDesc));
 		}
 		for (int32 OutputTensorIndex : Format.Operators[Idx].OutTensors)
 		{
-			OpOutputTensors.Emplace(AllTensors[OutputTensorIndex]);
+			FSymbolicTensorDesc SymbolicTensorDesc = AllSymbolicTensors[OutputTensorIndex];
+			//TODO jira 168972: Handle dynamic tensor desc, op should init from symbolic shapes
+			OpOutputTensors.Emplace(FMLTensorDesc::MakeFromSymbolic(SymbolicTensorDesc));
 		}
 		for (const FMLFormatAttributeDesc& Desc : Format.Operators[Idx].Attributes)
 		{
@@ -1299,7 +1297,7 @@ bool FMLInferenceModelDml::Init(UMLInferenceModel* InModel, FDeviceContextDml* I
 
 		FMLOperatorDml* Op = OpCreate(TypeName, OpInputTensors, OpOutputTensors, AttributeMap);
 
-		if (!Op)
+		if (!Op) //Op.Shader.IsNull())
 		{
 			UE_LOG(LogNNX, Warning, TEXT("Failed to create operator:%s"), *TypeName);
 
@@ -1316,14 +1314,40 @@ bool FMLInferenceModelDml::Init(UMLInferenceModel* InModel, FDeviceContextDml* I
 //
 //
 //
-void FMLInferenceModelDml::AddDispatchOps_RenderThread(FRDGBuilder& GraphBuilder, TArrayView<const FMLTensorBinding> InInputBindings, TArrayView<const FMLTensorBinding> OutOutputBindings)
+void FMLInferenceModelDml::AddDispatchOps_RenderThread(FRDGBuilder& GraphBuilder)
 {
+	check(AllTensorBindings.Num() == AllTensors.Num());
+	checkCode(
+		for (int32 i = 0; i < AllTensorBindings.Num(); ++i)
+		{
+			check(AllTensors[i].DataSize == AllTensorBindings[i].SizeInBytes);
+		}
+	);
+
+	static constexpr int32 MaxExpectedInput = 10;
+	TArray<FMLTensorBinding, TInlineAllocator<MaxExpectedInput>> InputBindings;
+	
+	static constexpr int32 MaxExpectedOutput = 2;
+	TArray<FMLTensorBinding, TInlineAllocator<MaxExpectedOutput>> OutputBindings;
+	
 	// Add passes for all operators
 	for (int32 Idx = 0; Idx < Operators.Num(); ++Idx)
 	{
+		InputBindings.Empty();
+		for (int32 i : OperatorInputTensorIndices[Idx])
+		{
+			InputBindings.Emplace(AllTensorBindings[i]);
+		}
+		OutputBindings.Empty();
+		for (int32 i : OperatorOutputTensorIndices[Idx])
+		{
+			OutputBindings.Emplace(AllTensorBindings[i]);
+		}
+
 		FMLOperatorDml* Op = Operators[Idx];
 
-		Op->Dispatch(GraphBuilder, InInputBindings, OutOutputBindings);
+		//TODO jira 169354 pass shape information to Op.Dispatch probably via an object containing both mem info and shape (from AllTensors).
+		Op->Dispatch(GraphBuilder, InputBindings, OutputBindings);
 	}
 }
 
diff --git a/Engine/Plugins/Experimental/NNX/Source/NNXRuntimeRDG/Private/NNXRuntimeHLSL.cpp b/Engine/Plugins/Experimental/NNX/Source/NNXRuntimeRDG/Private/NNXRuntimeHLSL.cpp
--- a/Engine/Plugins/Experimental/NNX/Source/NNXRuntimeRDG/Private/NNXRuntimeHLSL.cpp
+++ b/Engine/Plugins/Experimental/NNX/Source/NNXRuntimeRDG/Private/NNXRuntimeHLSL.cpp
@@ -34,7 +34,7 @@ public:
 
 protected:
 
-	virtual void AddDispatchOps_RenderThread(FRDGBuilder& GraphBuilder, TArrayView<const FMLTensorBinding> InInputBindings, TArrayView<const FMLTensorBinding> OutOutputBindings) override;
+	virtual void AddDispatchOps_RenderThread(FRDGBuilder& GraphBuilder) override;
 
 private:
 
@@ -83,11 +83,15 @@ bool FMLInferenceModelHlsl::Init(UMLInferenceModel* InModel)
 
 		for (int32 InputTensorIndex : Format.Operators[Idx].InTensors)
 		{
-			OpInputTensors.Emplace(AllTensors[InputTensorIndex]);
+			FSymbolicTensorDesc SymbolicTensorDesc = AllSymbolicTensors[InputTensorIndex];
+			//TODO jira 168972: Handle dynamic tensor desc, op should init from symbolic shapes
+			OpInputTensors.Emplace(FMLTensorDesc::MakeFromSymbolic(SymbolicTensorDesc));
 		}
 		for (int32 OutputTensorIndex : Format.Operators[Idx].OutTensors)
 		{
-			OpOutputTensors.Emplace(AllTensors[OutputTensorIndex]);
+			FSymbolicTensorDesc SymbolicTensorDesc = AllSymbolicTensors[OutputTensorIndex];
+			//TODO jira 168972: Handle dynamic tensor desc, op should init from symbolic shapes
+			OpOutputTensors.Emplace(FMLTensorDesc::MakeFromSymbolic(SymbolicTensorDesc));
 		}
 		for (const FMLFormatAttributeDesc& Desc : Format.Operators[Idx].Attributes)
 		{
@@ -133,14 +137,40 @@ FMLOperatorHlsl* FMLInferenceModelHlsl::OpCreate(const FString& OpName, TArrayVi
 //
 //
 //
-void FMLInferenceModelHlsl::AddDispatchOps_RenderThread(FRDGBuilder& GraphBuilder, TArrayView<const FMLTensorBinding> InInputBindings, TArrayView<const FMLTensorBinding> OutOutputBindings)
+void FMLInferenceModelHlsl::AddDispatchOps_RenderThread(FRDGBuilder& GraphBuilder)
 {
+	check(AllTensorBindings.Num() == AllTensors.Num());
+	checkCode(
+		for (int32 i = 0; i < AllTensorBindings.Num(); ++i)
+		{
+			check(AllTensors[i].DataSize == AllTensorBindings[i].SizeInBytes);
+		}
+	);
+
+	static constexpr int32 MaxExpectedInput = 10;
+	TArray<FMLTensorBinding, TInlineAllocator<MaxExpectedInput>> InputBindings;
+
+	static constexpr int32 MaxExpectedOutput = 2;
+	TArray<FMLTensorBinding, TInlineAllocator<MaxExpectedOutput>> OutputBindings;
+
 	// Add passes for all operators
 	for (int32 Idx = 0; Idx < Operators.Num(); ++Idx)
 	{
+		InputBindings.Empty();
+		for (int32 i : OperatorInputTensorIndices[Idx])
+		{
+			InputBindings.Emplace(AllTensorBindings[i]);
+		}
+		OutputBindings.Empty();
+		for (int32 i : OperatorOutputTensorIndices[Idx])
+		{
+			OutputBindings.Emplace(AllTensorBindings[i]);
+		}
+
 		FMLOperatorHlsl* Op = Operators[Idx];
 
-		Op->Dispatch(GraphBuilder, InInputBindings, OutOutputBindings);
+		//TODO jira 169354 pass shape information to Op.Dispatch probably via an object containing both mem info and shape (from AllTensors).
+		Op->Dispatch(GraphBuilder, InputBindings, OutputBindings);
 	}
 }
 
diff --git a/Engine/Plugins/Experimental/NNX/Source/NNXRuntimeRDG/Private/NNXRuntimeHLSLHelper.cpp b/Engine/Plugins/Experimental/NNX/Source/NNXRuntimeRDG/Private/NNXRuntimeHLSLHelper.cpp
--- a/Engine/Plugins/Experimental/NNX/Source/NNXRuntimeRDG/Private/NNXRuntimeHLSLHelper.cpp
+++ b/Engine/Plugins/Experimental/NNX/Source/NNXRuntimeRDG/Private/NNXRuntimeHLSLHelper.cpp
@@ -15,7 +15,7 @@ namespace NNX
 		checkf(TargetNumdimensionForBroadcast >= TensorDesc.Shape.Num(), TEXT("Can't broadcast tensor from rank %d to rank %d, should be inferior or equal."), TensorDesc.Shape.Num(), TargetNumdimensionForBroadcast);
 		int32 Offset = TargetNumdimensionForBroadcast - TensorDesc.Shape.Num();
 
-		static_assert(FMLTensorDesc::MaxTensorDimension <= NXRT_TENSORSTRIDEINFO_MAX_NUM_DIMENSIONS);
+		static_assert(FTensorShape::MaxRank <= NXRT_TENSORSTRIDEINFO_MAX_NUM_DIMENSIONS);
 		for (int32 i = 7; i >= 0; --i)
 		{
 			if (i >= TargetNumdimensionForBroadcast || i < Offset)
diff --git a/Engine/Plugins/Experimental/NNX/Source/NNXRuntimeRDG/Private/NNXRuntimeRDG.cpp b/Engine/Plugins/Experimental/NNX/Source/NNXRuntimeRDG/Private/NNXRuntimeRDG.cpp
--- a/Engine/Plugins/Experimental/NNX/Source/NNXRuntimeRDG/Private/NNXRuntimeRDG.cpp
+++ b/Engine/Plugins/Experimental/NNX/Source/NNXRuntimeRDG/Private/NNXRuntimeRDG.cpp
@@ -193,52 +193,72 @@ bool FMLInferenceModelRDG::LoadModel(const FNNIModelRaw& InModel, FMLRuntimeForm
 
 	FMLRuntimeFormat::StaticStruct()->SerializeBin(Reader, &Format);
 
+	// Data for base class
+	InputSymbolicTensors.Empty();
+	OutputSymbolicTensors.Empty();
+	
+	// Data for RDG
+	AllSymbolicTensors.Empty();
+	IntermediateTensorIndices.Empty();
+	InputTensorIndices.Empty();
+	OutputTensorIndices.Empty();
+	OperatorInputTensorIndices.Empty();
+	OperatorOutputTensorIndices.Empty();
+
 	// Add tensors
 	for (int32 Idx = 0; Idx < Format.Tensors.Num(); ++Idx)
 	{
 		const FMLFormatTensorDesc& FormatTensorDesc = Format.Tensors[Idx];
 
-		// When handling dynamic input shape FMLTensorDesc should then contain a FSymbolicTensorShape
-        // while actual inference work on FConcreteTensorShape resolved by shape inference.
 		FSymbolicTensorShape SymbolicShape = FSymbolicTensorShape::Make(FormatTensorDesc.Shape);
-		check(SymbolicShape.IsConcrete());
-		FConcreteTensorShape ConcreteShape = FConcreteTensorShape::Make(SymbolicShape);
+		FSymbolicTensorDesc SymbolicTensor = FSymbolicTensorDesc::Make(FormatTensorDesc.Name, SymbolicShape, FormatTensorDesc.DataType);
+		AllSymbolicTensors.Emplace(SymbolicTensor);
 		
-		FMLTensorDesc Tensor = FMLTensorDesc::Make(FormatTensorDesc.Name, ConcreteShape, FormatTensorDesc.DataType);
-		
-		Tensor.DataSize = Tensor.GetElemByteSize() * Tensor.Volume();
-
 		if (FormatTensorDesc.Type == EMLFormatTensorType::Input)
 		{
-			InputTensors.Add(Tensor);
-			AllTensors.Add(Tensor);
+			InputTensorIndices.Emplace(Idx);
+			InputSymbolicTensors.Emplace(SymbolicTensor);
 		}
 		else if (FormatTensorDesc.Type == EMLFormatTensorType::Output)
 		{
-			OutputTensors.Add(Tensor);
-			AllTensors.Add(Tensor);
+			OutputTensorIndices.Emplace(Idx);
+			OutputSymbolicTensors.Emplace(SymbolicTensor);
 		}
 		else if (FormatTensorDesc.Type == EMLFormatTensorType::Intermediate)
 		{
-			AllTensors.Add(Tensor);
+			IntermediateTensorIndices.Emplace(Idx);
 		}
 		checkf(FormatTensorDesc.Type != EMLFormatTensorType::None, TEXT("Unsupported tensor type None"));
 	}
 
+	// Loop over all operators in the model and store tensor indices for input/output
+	for (int32 Idx = 0; Idx < Format.Operators.Num(); ++Idx)
+	{
+		OperatorInputTensorIndices.Emplace(Format.Operators[Idx].InTensors);
+		OperatorOutputTensorIndices.Emplace(Format.Operators[Idx].OutTensors);
+	}
+
 	return true;
 }
 
 /**
  * Run the inference model (synchronous version)
  */
-int FMLInferenceModelRDG::Run(TArrayView<const FMLTensorBinding> InInputBindings, TArrayView<const FMLTensorBinding> OutOutputBindings)
+int FMLInferenceModelRDG::Run(TConstArrayView<const FMLTensorBinding> InInputBindings, TConstArrayView<const FTensorShape> InInputShapes, TConstArrayView<const FMLTensorBinding> InOutputBindings)
 {
+	// Prepare the model for the concrete input shapes
+	int Res = SetInputShapes(InInputShapes);
+	if (Res != 0)
+	{
+		UE_LOG(LogNNX, Warning, TEXT("Model preparation failed with error code:%d"), Res);
+		return -1;
+	}
+	
 	FEvent* Signal = FGenericPlatformProcess::GetSynchEventFromPool(false);
-	int		Res = 0;
 
 	ENQUEUE_RENDER_COMMAND(FMLInferenceModel_Run)
 	(
-		[&Signal, &Res, this, InInputBindings, OutOutputBindings](FRHICommandListImmediate& RHICmdList)
+		[&Signal, &Res, this, InInputBindings, InInputShapes, InOutputBindings](FRHICommandListImmediate& RHICmdList)
 		{
 			TOptional<ERHIPipeline>		Pipeline = RHICmdList.GetPipeline();
 
@@ -247,12 +267,12 @@ int FMLInferenceModelRDG::Run(TArrayView<const FMLTensorBinding> InInputBindings
 				RHICmdList.SwitchPipeline(ERHIPipeline::Graphics);
 			}
 
-			FRDGBuilder	GraphBuilder(RHICmdList);
+			FRDGBuilder	RDGBuilder(RHICmdList);
 
-			Res = EnqueueRDG(GraphBuilder, InInputBindings, OutOutputBindings);
+			Res = EnqueueRDG(RDGBuilder, InInputBindings, InInputShapes, InOutputBindings);
 			if (Res == 0)
 			{
-				GraphBuilder.Execute();
+				RDGBuilder.Execute();
 
 				// FIXME: Using BlockUntilGPUIdle() prevents hang on Linux
 				RHICmdList.BlockUntilGPUIdle();
@@ -285,20 +305,72 @@ int FMLInferenceModelRDG::Run(TArrayView<const FMLTensorBinding> InInputBindings
 	return Res;
 }
 
+int FMLInferenceModelRDG::SetInputShapes(TConstArrayView<const FTensorShape> InInputShapes)
+{
+	//Verify input shape are valid for the model
+	if (FMLInferenceModel::SetInputShapes(InInputShapes) != 0)
+	{
+		return -1;
+	}
+
+	//Run shape inference
+	AllTensors.Empty();
+	for (FSymbolicTensorDesc SymbolicTensorDesc : AllSymbolicTensors)
+	{
+		//Fake shape inference for now
+		check(SymbolicTensorDesc.IsConcrete());
+		FMLTensorDesc TensorDesc = FMLTensorDesc::MakeFromSymbolic(SymbolicTensorDesc);
+		AllTensors.Emplace(TensorDesc);
+	}
+
+	//Build concrete input and output tensor descriptions
+	InputTensors.Empty();
+	for (int32 InputIndices : InputTensorIndices)
+	{
+		InputTensors.Emplace(AllTensors[InputIndices]);
+	}
+	OutputTensors.Empty();
+	for (int32 OutputIndices : OutputTensorIndices)
+	{
+		OutputTensors.Emplace(AllTensors[OutputIndices]);
+	}
+	check(InputTensorIndices.Num() + OutputTensorIndices.Num() + IntermediateTensorIndices.Num() == AllTensors.Num());
+	check(InputTensors.Num() + OutputTensors.Num() + IntermediateTensorIndices.Num() == AllSymbolicTensors.Num());
+	
+	return 0;
+}
+
+FRDGBufferDesc CreateRDGBufferDescFromTensorDesc(const FMLTensorDesc& TensorDesc)
+{
+	// FIXME: CreateStructuredDesc() creates a crash on VulkanRHI
+	//FRDGBufferDesc Desc = FRDGBufferDesc::CreateStructuredDesc(TensorDesc.GetElemByteSize(), TensorDesc.Volume);
+	FRDGBufferDesc Desc = FRDGBufferDesc::CreateBufferDesc(TensorDesc.GetElemByteSize(), TensorDesc.Volume);
+
+	return Desc;
+}
+
 /**
  * Enqueue operators to RDG, the caller will run the GraphBuilder.Execute()
  */
-int FMLInferenceModelRDG::EnqueueRDG(FRDGBuilder& GraphBuilder, TArrayView<const FMLTensorBinding> InInputBindings, TArrayView<const FMLTensorBinding> InOutputBindings)
+int FMLInferenceModelRDG::EnqueueRDG(FRDGBuilder& RDGBuilder, TConstArrayView<const FMLTensorBinding> InInputBindings, TConstArrayView<const FTensorShape> InInputShapes, TConstArrayView<const FMLTensorBinding> InOutputBindings)
 {
 	check(IsInRenderingThread());
 
 	int Res;
 
+	// Prepare the model for the concrete input shapes
+	Res = SetInputShapes(InInputShapes);
+	if (Res != 0)
+	{
+		UE_LOG(LogNNX, Warning, TEXT("Model preparation failed with error code:%d"), Res);
+		return -1;
+	}
+
 	// Process input tensors, and if required, allocate RDG buffers
 	FMLTensorBindingArray	RDGInputBindings;
 	FMLIntArray				RDGUploadIndices;
 
-	Res = SetTensors(GraphBuilder, RDGInputBindings, RDGUploadIndices, InInputBindings, InputTensors);
+	Res = SetTensors(RDGBuilder, RDGInputBindings, RDGUploadIndices, InInputBindings, InputTensors);
 	if (Res != 0)
 	{
 		UE_LOG(LogNNX, Warning, TEXT("Invalid input tensor binding type for tensor index:%d"), Res);
@@ -309,7 +381,7 @@ int FMLInferenceModelRDG::EnqueueRDG(FRDGBuilder& GraphBuilder, TArrayView<const
 	FMLTensorBindingArray	RDGOutputBindings;
 	FMLIntArray				RDGReadbackIndices;
 
-	Res = SetTensors(GraphBuilder, RDGOutputBindings, RDGReadbackIndices, InOutputBindings, OutputTensors);
+	Res = SetTensors(RDGBuilder, RDGOutputBindings, RDGReadbackIndices, InOutputBindings, OutputTensors);
 	if (Res != 0)
 	{
 		UE_LOG(LogNNX, Warning, TEXT("Invalid output tensor binding type for tensor index:%d"), Res);
@@ -319,16 +391,53 @@ int FMLInferenceModelRDG::EnqueueRDG(FRDGBuilder& GraphBuilder, TArrayView<const
 	// If required, upload input tensors to GPU
 	if (!RDGUploadIndices.IsEmpty())
 	{
-		AddTensorUploads_RenderThread(GraphBuilder, RDGUploadIndices, RDGInputBindings, InInputBindings);
+		AddTensorUploads_RenderThread(RDGBuilder, RDGUploadIndices, RDGInputBindings, InInputBindings);
 	}
 
+	AllTensorBindings.Empty();
+	AllTensorBindings.SetNumUninitialized(AllTensors.Num());
+	checkCode(
+		for (int32 i = 0; i < AllTensorBindings.Num(); ++i)
+		{
+			AllTensorBindings[i].SizeInBytes = 0;
+		}
+	);
+	
+	//Create intermediate tensors bindings
+	for (int32 Idx : IntermediateTensorIndices)
+	{
+		const FMLTensorDesc& TensorDesc = AllTensors[Idx];
+		FRDGBufferDesc Desc = CreateRDGBufferDescFromTensorDesc(TensorDesc);
+		FRDGBufferRef TensorBuffer = RDGBuilder.CreateBuffer(Desc, *TensorDesc.Name, ERDGBufferFlags::None);
+		AllTensorBindings[Idx] = FMLTensorBinding::FromRDG(TensorBuffer, TensorDesc.DataSize);
+	}
+
+	//Insert input tensors bindings
+	for (int32 i = 0; i < InputTensorIndices.Num(); ++i)
+	{
+		AllTensorBindings[InputTensorIndices[i]] = RDGInputBindings[i];
+	}
+
+	//Insert output tensors bindings
+	for (int32 i = 0; i < OutputTensorIndices.Num(); ++i)
+	{
+		AllTensorBindings[OutputTensorIndices[i]] = RDGOutputBindings[i];
+	}
+
+	checkCode(
+		for (int32 i = 0; i < AllTensorBindings.Num(); ++i)
+		{
+			check(AllTensorBindings[i].SizeInBytes != 0);
+		}
+	);
+	
 	// We can now dispatch operators
-	AddDispatchOps_RenderThread(GraphBuilder, RDGInputBindings, RDGOutputBindings);
+	AddDispatchOps_RenderThread(RDGBuilder);
 
 	// If required, readback the output tensors to CPU
 	if (!RDGReadbackIndices.IsEmpty())
 	{
-		AddTensorReadbacks_RenderThread(GraphBuilder, RDGReadbackIndices, RDGOutputBindings, InOutputBindings);		
+		AddTensorReadbacks_RenderThread(RDGBuilder, RDGReadbackIndices, RDGOutputBindings, InOutputBindings);
 	}
 
 	return 0;
@@ -347,9 +456,7 @@ int FMLInferenceModelRDG::SetTensors(FRDGBuilder& GraphBuilder, FMLTensorBinding
 
 		if (Binding.BindingType == EMLTensorBindingDataType::CPUMemory)
 		{
-			// FIXME: CreateStructuredDesc() creates a crash on VulkanRHI
-			//FRDGBufferDesc Desc = FRDGBufferDesc::CreateStructuredDesc(TensorDesc.GetElemByteSize(), TensorDesc.Num());
-			FRDGBufferDesc Desc = FRDGBufferDesc::CreateBufferDesc(TensorDesc.GetElemByteSize(), TensorDesc.Num());
+			FRDGBufferDesc Desc = CreateRDGBufferDescFromTensorDesc(InTensors[Idx]);
 
 			// FIXME: We should use BUF_SourceCopy for only output buffers (GPU readback)
 			Desc.Usage = EBufferUsageFlags(Desc.Usage | BUF_SourceCopy);
diff --git a/Engine/Plugins/Experimental/NNX/Source/NNXRuntimeRDG/Private/NNXRuntimeRDG.h b/Engine/Plugins/Experimental/NNX/Source/NNXRuntimeRDG/Private/NNXRuntimeRDG.h
--- a/Engine/Plugins/Experimental/NNX/Source/NNXRuntimeRDG/Private/NNXRuntimeRDG.h
+++ b/Engine/Plugins/Experimental/NNX/Source/NNXRuntimeRDG/Private/NNXRuntimeRDG.h
@@ -63,8 +63,9 @@ public:
 
 	~FMLInferenceModelRDG();
 
-	virtual int Run(TArrayView<const FMLTensorBinding> InInputBindings, TArrayView<const FMLTensorBinding> OutOutputBindings) override;
-	virtual int EnqueueRDG(FRDGBuilder& Builder, TArrayView<const FMLTensorBinding> InInputBindings, TArrayView<const FMLTensorBinding> OutOutputBindings) override;
+	virtual int SetInputShapes(TConstArrayView<const FTensorShape> InputShapes) override;
+	virtual int Run(TConstArrayView<const FMLTensorBinding> InputBindings, TConstArrayView<const FTensorShape> InputShapes, TConstArrayView<const FMLTensorBinding> OutputBindings) override;
+	virtual int EnqueueRDG(FRDGBuilder& RDGBuilder, TConstArrayView<const FMLTensorBinding> InputBindings, TConstArrayView<const FTensorShape> InputShapes, TConstArrayView<const FMLTensorBinding> OutputBindings) override;
 
 protected:
 
@@ -72,15 +73,32 @@ protected:
 
 	bool LoadModel(const FNNIModelRaw& InModel, FMLRuntimeFormat& Format);
 
+	int CreateIntermediateTensors(FRDGBuilder& GraphBuilder, FMLTensorBindingArray& OutBindings, TConstArrayView<const FMLTensorDesc> InTensors);
 	int SetTensors(FRDGBuilder& GraphBuilder, FMLTensorBindingArray& OutRDGBindings, FMLIntArray& OutIndices, TArrayView<const FMLTensorBinding> InBindings, TArrayView<const FMLTensorDesc> InTensors);
 
-	virtual void AddDispatchOps_RenderThread(FRDGBuilder& GraphBuilder, TArrayView<const FMLTensorBinding> InInputBindings, TArrayView<const FMLTensorBinding> OutOutputBindings) = 0;
+	virtual void AddDispatchOps_RenderThread(FRDGBuilder& GraphBuilder) = 0;
 
 	virtual void AddTensorUploads_RenderThread(FRDGBuilder& GraphBuilder, TArrayView<const int32> InUploadIndices, TArrayView<FMLTensorBinding> InRDGBindings, TArrayView<const FMLTensorBinding> InBindings);
 	virtual void AddTensorReadbacks_RenderThread(FRDGBuilder& GraphBuilder, TArrayView<const int32> InReadbackIndices, TArrayView<const FMLTensorBinding> InRDGBindings, TArrayView<const FMLTensorBinding> InBindings);
 
-	FReadbackEntry	Readback;
-	bool			bUseManualTransitions;
+	//Tensors
+	TArray<FMLTensorDesc>		AllTensors;
+	TArray<FSymbolicTensorDesc>	AllSymbolicTensors;
+	TArray<FMLTensorDesc>		InputTensors;
+	TArray<FMLTensorDesc>		OutputTensors;
+
+	//Tensor indices for models
+	TArray<int32>				IntermediateTensorIndices;
+	TArray<int32>				InputTensorIndices;
+	TArray<int32>				OutputTensorIndices;
+	
+	//Tensor indices by operator
+	TArray<TArray<uint32>>		OperatorInputTensorIndices;
+	TArray<TArray<uint32>>		OperatorOutputTensorIndices;
+	
+	FMLTensorBindingArray       AllTensorBindings;
+	FReadbackEntry				Readback;
+	bool						bUseManualTransitions;
 };
 
 //TODO jira 167585 remove default validation and declare contract in all HLSL operator (see HLSL Gemm for current example)
diff --git a/Engine/Plugins/Experimental/NNX/Source/NNXUtils/Private/NNXModelBuilderNNX.cpp b/Engine/Plugins/Experimental/NNX/Source/NNXUtils/Private/NNXModelBuilderNNX.cpp
--- a/Engine/Plugins/Experimental/NNX/Source/NNXUtils/Private/NNXModelBuilderNNX.cpp
+++ b/Engine/Plugins/Experimental/NNX/Source/NNXUtils/Private/NNXModelBuilderNNX.cpp
@@ -129,7 +129,7 @@ public:
 		
 		if (Format.Tensors[Idx].Type != EMLFormatTensorType::None)
 		{
-			UE_LOG(LogNNX, Warning, TEXT("Failed to add input tensor, tensor usage already setupped"));
+			UE_LOG(LogNNX, Warning, TEXT("Failed to add input tensor, tensor usage already set up"));
 			return false;
 		}
 		
@@ -151,7 +151,7 @@ public:
 
 		if (Format.Tensors[Idx].Type != EMLFormatTensorType::None)
 		{
-			UE_LOG(LogNNX, Warning, TEXT("Failed to add output tensor, tensor usage already setupped"));
+			UE_LOG(LogNNX, Warning, TEXT("Failed to add output tensor, tensor usage already set up"));
 			return false;
 		}
 
@@ -204,7 +204,7 @@ public:
 
 		if (Format.Tensors[TensorIdx].Type == EMLFormatTensorType::Input)
 		{
-			UE_LOG(LogNNX, Warning, TEXT("Failed to add output tensor, tensor usage already setupped to input"));
+			UE_LOG(LogNNX, Warning, TEXT("Failed to add output tensor, tensor usage already set up to input"));
 			return false;
 		}
 
diff --git a/Engine/Plugins/Experimental/NNX/Source/NNXUtils/Private/NNXUtilsModule.cpp b/Engine/Plugins/Experimental/NNX/Source/NNXUtils/Private/NNXUtilsModule.cpp
--- a/Engine/Plugins/Experimental/NNX/Source/NNXUtils/Private/NNXUtilsModule.cpp
+++ b/Engine/Plugins/Experimental/NNX/Source/NNXUtils/Private/NNXUtilsModule.cpp
@@ -91,8 +91,8 @@ static FAutoConsoleCommand ConsoleCmdNNXOptimizeModel(
 			}
 			else if (!OpName.IsEmpty())
 			{
-				FMLTensorDesc InputTensor = FMLTensorDesc::Make(TEXT("in"), FConcreteTensorShape::Make({1, 512 }), EMLTensorDataType::Float);
-				FMLTensorDesc OutputTensor = FMLTensorDesc::Make(TEXT("out"), FConcreteTensorShape::Make({ 1, 512 }), EMLTensorDataType::Float);
+				FMLTensorDesc InputTensor = FMLTensorDesc::Make(TEXT("in"), FTensorShape::Make({1, 512 }), EMLTensorDataType::Float);
+				FMLTensorDesc OutputTensor = FMLTensorDesc::Make(TEXT("out"), FTensorShape::Make({ 1, 512 }), EMLTensorDataType::Float);
 
 				if (!NNX::CreateONNXModelForOperator(OpName, MakeArrayView(&InputTensor, 1), MakeArrayView(&OutputTensor, 1), ONNXModel))
 				{
@@ -140,8 +140,8 @@ static FAutoConsoleCommand ConsoleCmdNNXCreateModel(
 			if (Args.Num() > 0)
 			{
 				FNNIModelRaw ONNXModel;
-				FMLTensorDesc InputTensor = FMLTensorDesc::Make(TEXT("in"), FConcreteTensorShape::Make({ 1, 512 }), EMLTensorDataType::Float);
-				FMLTensorDesc OutputTensor = FMLTensorDesc::Make(TEXT("out"), FConcreteTensorShape::Make({ 1, 512 }), EMLTensorDataType::Float);
+				FMLTensorDesc InputTensor = FMLTensorDesc::Make(TEXT("in"), FTensorShape::Make({ 1, 512 }), EMLTensorDataType::Float);
+				FMLTensorDesc OutputTensor = FMLTensorDesc::Make(TEXT("out"), FTensorShape::Make({ 1, 512 }), EMLTensorDataType::Float);
 
 				if (!NNX::CreateONNXModelForOperator(Args[0], MakeArrayView(&InputTensor, 1), MakeArrayView(&OutputTensor, 1), ONNXModel))
 				{
