NNI Cleanup

* FSymbolicTensorDesc renamed to FTensorDesc
* FTensorDesc renamed to FTensor (to evolve further)
* Shape.Num() renamed to Shape.Rank()
* Change FTensorDesc and FTensor from aggregate/pod to class to ensure internal consistency
* Remove const in TConstArray<> as it is implicit

diff --git a/Engine/Plugins/Experimental/NNX/Source/NNEHlslShaders/Internal/NNEHlslShadersGatherCS.h b/Engine/Plugins/Experimental/NNX/Source/NNEHlslShaders/Internal/NNEHlslShadersGatherCS.h
--- a/Engine/Plugins/Experimental/NNX/Source/NNEHlslShaders/Internal/NNEHlslShadersGatherCS.h
+++ b/Engine/Plugins/Experimental/NNX/Source/NNEHlslShaders/Internal/NNEHlslShadersGatherCS.h
@@ -42,7 +42,7 @@ namespace UE::NNEHlslShaders::Internal
 
 		static void ModifyCompilationEnvironment(const FGlobalShaderPermutationParameters& InParameters, FShaderCompilerEnvironment& OutEnvironment);
 
-		static void FillInParameters(int32 Axis, const NNX::FMLTensorDesc& Data, const NNX::FMLTensorDesc& Indices, FParameters& Parameters);
+		static void FillInParameters(int32 Axis, const NNX::FTensor& Data, const NNX::FTensor& Indices, FParameters& Parameters);
 
 		static FIntVector GetGroupCount(const FParameters& Parameters);
 	};
diff --git a/Engine/Plugins/Experimental/NNX/Source/NNEHlslShaders/Internal/NNEHlslShadersGemmCS.h b/Engine/Plugins/Experimental/NNX/Source/NNEHlslShaders/Internal/NNEHlslShadersGemmCS.h
--- a/Engine/Plugins/Experimental/NNX/Source/NNEHlslShaders/Internal/NNEHlslShadersGemmCS.h
+++ b/Engine/Plugins/Experimental/NNX/Source/NNEHlslShaders/Internal/NNEHlslShadersGemmCS.h
@@ -79,9 +79,9 @@ namespace UE::NNEHlslShaders::Internal
 
 		static void ModifyCompilationEnvironment(const FGlobalShaderPermutationParameters& InParameters, FShaderCompilerEnvironment& OutEnvironment);
 
-		static void FillInParameters(float Alpha, float Beta, int32 TransA, int32 TransB, const NNX::FMLTensorDesc &InputA, const NNX::FMLTensorDesc &InputB, const NNX::FMLTensorDesc &InputC, float CScalar, FParameters& Parameters);
+		static void FillInParameters(float Alpha, float Beta, int32 TransA, int32 TransB, const NNX::FTensor& InputA, const NNX::FTensor& InputB, const NNX::FTensor& InputC, float CScalar, FParameters& Parameters);
 
-		static void FillInParametersMatMul(const NNX::FMLTensorDesc &InputA, const NNX::FMLTensorDesc &InputB, FParameters& Parameters);
+		static void FillInParametersMatMul(const NNX::FTensor& InputA, const NNX::FTensor& InputB, FParameters& Parameters);
 
 		static FIntVector GetGroupCount(const FParameters& Parameters, EGemmAlgorithm Algorithm, int32 NumStackDimensions);
 		static EGemmAlgorithm GetAlgorithm(const FParameters& Parameters);
diff --git a/Engine/Plugins/Experimental/NNX/Source/NNEHlslShaders/Private/NNEHlslShadersGatherCS.cpp b/Engine/Plugins/Experimental/NNX/Source/NNEHlslShaders/Private/NNEHlslShadersGatherCS.cpp
--- a/Engine/Plugins/Experimental/NNX/Source/NNEHlslShaders/Private/NNEHlslShadersGatherCS.cpp
+++ b/Engine/Plugins/Experimental/NNX/Source/NNEHlslShaders/Private/NNEHlslShadersGatherCS.cpp
@@ -15,47 +15,47 @@ namespace UE::NNEHlslShaders::Internal
 	}
 
 	//template <typename DataElementType, typename IndicesElementType>
-	void TGatherCS::FillInParameters(int32 Axis, const NNX::FMLTensorDesc& Data, const NNX::FMLTensorDesc& Indices, FParameters& Parameters)
+	void TGatherCS::FillInParameters(int32 Axis, const NNX::FTensor& Data, const NNX::FTensor& Indices, FParameters& Parameters)
 	{
 		Parameters.Axis = Axis;
 
 		TArray<int32> OutputShape;
-		OutputShape.SetNum(Indices.Shape.Num() + (Data.Shape.Num() - 1));
+		OutputShape.SetNum(Indices.GetShape().Rank() + (Data.GetShape().Rank() - 1));
 		int32 index = 0;
 		Parameters.OutputSize = 1;
 		for (int32 i = 0; i < Axis; i++, index++)
 		{
-			OutputShape[index] = Data.Shape.Data[i];
+			OutputShape[index] = Data.GetShape().Data[i];
 			Parameters.OutputSize *= OutputShape[index];
 		}
-		for (int32 i = 0; i < Indices.Shape.Num(); i++, index++)
+		for (int32 i = 0; i < Indices.GetShape().Rank(); i++, index++)
 		{
-			OutputShape[index] = Indices.Shape.Data[i];
+			OutputShape[index] = Indices.GetShape().Data[i];
 			Parameters.OutputSize *= OutputShape[index];
 		}
-		for (int32 i = (Axis + 1); i < Data.Shape.Num(); i++, index++)
+		for (int32 i = (Axis + 1); i < Data.GetShape().Rank(); i++, index++)
 		{
-			OutputShape[index] = Data.Shape.Data[i];
+			OutputShape[index] = Data.GetShape().Data[i];
 			Parameters.OutputSize *= OutputShape[index];
 		}
 
-		Parameters.NumDataDimensions = Data.Shape.Num();
-		Parameters.NumIndicesDimensions = Indices.Shape.Num();
+		Parameters.NumDataDimensions = Data.GetShape().Rank();
+		Parameters.NumIndicesDimensions = Indices.GetShape().Rank();
 
 		int32 DataStride = 1;
-		for (int32 i = Data.Shape.Num() - 1; i >= 0; i--)
+		for (int32 i = Data.GetShape().Rank() - 1; i >= 0; i--)
 		{
 			Parameters.DataStride_IndicesStride_OutputStride[i].X = DataStride;
 			Parameters.OneDivDataStride_OneDivIndicesStride_OneDivOutputStride[i].X = 1.0 / (float)DataStride;
-			DataStride *= Data.Shape.Data[i];
+			DataStride *= Data.GetShape().Data[i];
 		}
 
 		int32 IndicesStride = 1;
-		for (int32 i = Indices.Shape.Num() - 1; i >= 0; i--)
+		for (int32 i = Indices.GetShape().Rank() - 1; i >= 0; i--)
 		{
 			Parameters.DataStride_IndicesStride_OutputStride[i].Y = IndicesStride;
 			Parameters.OneDivDataStride_OneDivIndicesStride_OneDivOutputStride[i].Y = 1.0 / (float)IndicesStride;
-			IndicesStride *= Indices.Shape.Data[i];
+			IndicesStride *= Indices.GetShape().Data[i];
 		}
 
 		int32 OutputStride = 1;
diff --git a/Engine/Plugins/Experimental/NNX/Source/NNEHlslShaders/Private/NNEHlslShadersGemmCS.cpp b/Engine/Plugins/Experimental/NNX/Source/NNEHlslShaders/Private/NNEHlslShadersGemmCS.cpp
--- a/Engine/Plugins/Experimental/NNX/Source/NNEHlslShaders/Private/NNEHlslShadersGemmCS.cpp
+++ b/Engine/Plugins/Experimental/NNX/Source/NNEHlslShaders/Private/NNEHlslShadersGemmCS.cpp
@@ -17,14 +17,14 @@ namespace UE::NNEHlslShaders::Internal
 		TArray<uint32> StackStrideB;
 	};
 
-	static GemmMatrixParameters GetMatrixParametersMatMul(const NNX::FMLTensorDesc& InputA, const NNX::FMLTensorDesc& InputB)
+	static GemmMatrixParameters GetMatrixParametersMatMul(const NNX::FTensor& InputA, const NNX::FTensor& InputB)
 	{
-		check(InputA.Shape.Num() != 0);
-		check(InputB.Shape.Num() != 0);
+		check(InputA.GetShape().Rank() != 0);
+		check(InputB.GetShape().Rank() != 0);
 
 		GemmMatrixParameters Result;
-		int32 NumStackDimensionsA = FMath::Max((int32)InputA.Shape.Num() - 2, 0);
-		int32 NumStackDimensionsB = FMath::Max((int32)InputB.Shape.Num() - 2, 0);
+		int32 NumStackDimensionsA = FMath::Max((int32)InputA.GetShape().Rank() - 2, 0);
+		int32 NumStackDimensionsB = FMath::Max((int32)InputB.GetShape().Rank() - 2, 0);
 		int32 NumStackDimensions = FMath::Max(NumStackDimensionsA, NumStackDimensionsB);
 
 		check(NumStackDimensions <= FGemmConstants::MAX_NUM_STACK_DIMENSIONS);
@@ -35,31 +35,31 @@ namespace UE::NNEHlslShaders::Internal
 			const int32 NumDimToCheck = FMath::Min(NumStackDimensionsA, NumStackDimensionsB);
 			for (int32 i = 0; i < NumDimToCheck; i++)
 			{
-				const uint32 VolumeA = InputA.Shape.Data[NumStackDimensionsA - 1 - i];
-				const uint32 VolumeB = InputB.Shape.Data[NumStackDimensionsB - 1 - i];
+				const uint32 VolumeA = InputA.GetShape().Data[NumStackDimensionsA - 1 - i];
+				const uint32 VolumeB = InputB.GetShape().Data[NumStackDimensionsB - 1 - i];
 
 				check(VolumeA == 1 || VolumeB == 1 || VolumeA == VolumeB);
 			}
 		}
 
-		const bool IsVectorA = InputA.Shape.Num() == 1;
-		const bool IsVectorB = InputB.Shape.Num() == 1;
+		const bool IsVectorA = InputA.GetShape().Rank() == 1;
+		const bool IsVectorB = InputB.GetShape().Rank() == 1;
 
-		Result.M = IsVectorA ? 1 : InputA.Shape.Data[InputA.Shape.Num() - 2];
-		Result.N = IsVectorB ? 1 : InputB.Shape.Data[InputB.Shape.Num() - 1];
-		Result.K = InputA.Shape.Data[IsVectorA ? 0 : InputA.Shape.Num() - 1];
-		check(InputB.Shape.Data[IsVectorB ? 0 : InputB.Shape.Num() - 2] == Result.K);
+		Result.M = IsVectorA ? 1 : InputA.GetShape().Data[InputA.GetShape().Rank() - 2];
+		Result.N = IsVectorB ? 1 : InputB.GetShape().Data[InputB.GetShape().Rank() - 1];
+		Result.K = InputA.GetShape().Data[IsVectorA ? 0 : InputA.GetShape().Rank() - 1];
+		check(InputB.GetShape().Data[IsVectorB ? 0 : InputB.GetShape().Rank() - 2] == Result.K);
 
 		Result.StackShapeA.Init(1, NumStackDimensions);
 		Result.StackShapeB.Init(1, NumStackDimensions);
 
 		for (int32 i = 0; i < NumStackDimensions; i++)
 		{
-			int32 IdxA = InputA.Shape.Num() - 3 - i;
-			int32 IdxB = InputB.Shape.Num() - 3 - i;
+			int32 IdxA = InputA.GetShape().Rank() - 3 - i;
+			int32 IdxB = InputB.GetShape().Rank() - 3 - i;
 
-			Result.StackShapeA[Result.StackShapeA.Num() - 1 - i] = IdxA >= 0 ? InputA.Shape.Data[IdxA] : 1;
-			Result.StackShapeB[Result.StackShapeB.Num() - 1 - i] = IdxB >= 0 ? InputB.Shape.Data[IdxB] : 1;
+			Result.StackShapeA[Result.StackShapeA.Num() - 1 - i] = IdxA >= 0 ? InputA.GetShape().Data[IdxA] : 1;
+			Result.StackShapeB[Result.StackShapeB.Num() - 1 - i] = IdxB >= 0 ? InputB.GetShape().Data[IdxB] : 1;
 		}
 
 		Result.StackStrideA.Init(1, NumStackDimensions);
@@ -80,15 +80,15 @@ namespace UE::NNEHlslShaders::Internal
 		OutEnvironment.SetDefine(TEXT("MAX_NUM_STACK_DIMENSIONS"), FGemmConstants::MAX_NUM_STACK_DIMENSIONS);
 	}
 
-	void TGemmCS::FillInParameters(float Alpha, float Beta, int32 TransA, int32 TransB, const NNX::FMLTensorDesc& InputA, const NNX::FMLTensorDesc& InputB, const NNX::FMLTensorDesc& InputC, float CScalar, TGemmCS::FParameters& Parameters)
+	void TGemmCS::FillInParameters(float Alpha, float Beta, int32 TransA, int32 TransB, const NNX::FTensor& InputA, const NNX::FTensor& InputB, const NNX::FTensor& InputC, float CScalar, TGemmCS::FParameters& Parameters)
 	{
-		check(InputA.Shape.Num() == 2);
-		check(InputB.Shape.Num() == 2);
+		check(InputA.GetShape().Rank() == 2);
+		check(InputB.GetShape().Rank() == 2);
 
-		uint32 M = TransA != 0 ? InputA.Shape.Data[1] : InputA.Shape.Data[0];
-		uint32 K = TransA != 0 ? InputA.Shape.Data[0] : InputA.Shape.Data[1];
-		uint32 N = TransB != 0 ? InputB.Shape.Data[0] : InputB.Shape.Data[1];
-		check(K == (TransB != 0 ? InputB.Shape.Data[1] : InputB.Shape.Data[0]));
+		uint32 M = TransA != 0 ? InputA.GetShape().Data[1] : InputA.GetShape().Data[0];
+		uint32 K = TransA != 0 ? InputA.GetShape().Data[0] : InputA.GetShape().Data[1];
+		uint32 N = TransB != 0 ? InputB.GetShape().Data[0] : InputB.GetShape().Data[1];
+		check(K == (TransB != 0 ? InputB.GetShape().Data[1] : InputB.GetShape().Data[0]));
 
 		Parameters.Alpha = Alpha;
 		Parameters.Beta = Beta;
@@ -100,12 +100,12 @@ namespace UE::NNEHlslShaders::Internal
 		Parameters.MxK = M * K;
 		Parameters.KxN = K * N;
 		Parameters.MxN = M * N;
-		Parameters.CWidth = InputC.Shape.Num() == 0 ? 0 : InputC.Shape.Data[InputC.Shape.Num() - 1];;
-		Parameters.CHeight = InputC.Shape.Num() == 0 ? 0 : InputC.Shape.Num() == 1 ? 1 : InputC.Shape.Data[InputC.Shape.Num() - 2];
+		Parameters.CWidth = InputC.GetShape().Rank() == 0 ? 0 : InputC.GetShape().Data[InputC.GetShape().Rank() - 1];;
+		Parameters.CHeight = InputC.GetShape().Rank() == 0 ? 0 : InputC.GetShape().Rank() == 1 ? 1 : InputC.GetShape().Data[InputC.GetShape().Rank() - 2];
 		Parameters.CScalar = CScalar;
 	}
 
-	void TGemmCS::FillInParametersMatMul(const NNX::FMLTensorDesc& InputA, const NNX::FMLTensorDesc& InputB, TGemmCS::FParameters& Parameters)
+	void TGemmCS::FillInParametersMatMul(const NNX::FTensor& InputA, const NNX::FTensor& InputB, TGemmCS::FParameters& Parameters)
 	{
 		GemmMatrixParameters Params = GetMatrixParametersMatMul(InputA, InputB);
 
diff --git a/Engine/Plugins/Experimental/NNX/Source/NNXCore/Private/NNXRuntime.cpp b/Engine/Plugins/Experimental/NNX/Source/NNXCore/Private/NNXRuntime.cpp
--- a/Engine/Plugins/Experimental/NNX/Source/NNXCore/Private/NNXRuntime.cpp
+++ b/Engine/Plugins/Experimental/NNX/Source/NNXCore/Private/NNXRuntime.cpp
@@ -42,7 +42,7 @@ namespace NNX
 		return Binding;
 	}
 
-	int FMLInferenceModel::SetInputShapes(TConstArrayView<const FTensorShape> InInputShapes)
+	int FMLInferenceModel::SetInputShapes(TConstArrayView<FTensorShape> InInputShapes)
 	{
 		if (InInputShapes.Num() != InputSymbolicTensors.Num())
 		{
@@ -52,10 +52,10 @@ namespace NNX
 
 		for (int32 i = 0; i < InInputShapes.Num(); ++i)
 		{
-			const FSymbolicTensorDesc SymbolicDesc = InputSymbolicTensors[i];
-			if (!InInputShapes[i].IsCompatibleWith(SymbolicDesc.Shape))
+			const FTensorDesc SymbolicDesc = InputSymbolicTensors[i];
+			if (!InInputShapes[i].IsCompatibleWith(SymbolicDesc.GetShape()))
 			{
-				UE_LOG(LogNNX, Warning, TEXT("Input shape does not match input tensor %s of index %d"), *SymbolicDesc.Name, i);
+				UE_LOG(LogNNX, Warning, TEXT("Input shape does not match input tensor %s of index %d"), *SymbolicDesc.GetName(), i);
 				return -1;
 			}
 		}
diff --git a/Engine/Plugins/Experimental/NNX/Source/NNXCore/Public/NNXRuntime.h b/Engine/Plugins/Experimental/NNX/Source/NNXCore/Public/NNXRuntime.h
--- a/Engine/Plugins/Experimental/NNX/Source/NNXCore/Public/NNXRuntime.h
+++ b/Engine/Plugins/Experimental/NNX/Source/NNXCore/Public/NNXRuntime.h
@@ -103,29 +103,29 @@ public:
 	}
 
 	/** Getters for tensor description as defined by the model potentially with variable dimensions() */
-	TConstArrayView<const FSymbolicTensorDesc> GetInputTensorDescs() const;
-	TConstArrayView<const FSymbolicTensorDesc> GetOutputTensorDescs() const;
+	TConstArrayView<FTensorDesc> GetInputTensorDescs() const;
+	TConstArrayView<FTensorDesc> GetOutputTensorDescs() const;
 
 	/** Getters for input/output shapes if they were set already. Empty list otherwise. See SetInputShapes() */
-	TConstArrayView<const FTensorShape> GetInputTensorShapes() const;
-	TConstArrayView<const FTensorShape> GetOutputTensorShapes() const;
+	TConstArrayView<FTensorShape> GetInputTensorShapes() const;
+	TConstArrayView<FTensorShape> GetOutputTensorShapes() const;
 	
 	/** This call will prepare the model to be run with the given input shape, it is an optional call.
 	 * Run() and EnqueueRDG() will call it internally if it was not call explicitely.
 	 * It is recommended to call this function ahead of time to avoid overhead during the first model execution. */
-	virtual int SetInputShapes(TConstArrayView<const FTensorShape> InInputShapes);
+	virtual int SetInputShapes(TConstArrayView<FTensorShape> InInputShapes);
 
 	/** This call is synchronous on all engine types(CPU, RDG, GPU), i.e.the calling thread will be blocked.
 	 * until execution is finished */
-	virtual int Run(TConstArrayView<const FMLTensorBinding> InInputTensors,
-		            TConstArrayView<const FTensorShape> InInputShapes,
-					TConstArrayView<const FMLTensorBinding> InOutputTensors) = 0;
+	virtual int Run(TConstArrayView<FMLTensorBinding> InInputTensors,
+		            TConstArrayView<FTensorShape> InInputShapes,
+					TConstArrayView<FMLTensorBinding> InOutputTensors) = 0;
 
 	/** Enqueue the execution on the Render graph render thread. It's caller's responsibility to actually run the graph. */
 	virtual int EnqueueRDG(FRDGBuilder& RDGBuilder, 
-		TConstArrayView<const FMLTensorBinding> InInputTensors, 
-		TConstArrayView<const FTensorShape> InInputShapes,
-		TConstArrayView<const FMLTensorBinding> InOutputTensors)
+		TConstArrayView<FMLTensorBinding> InInputTensors, 
+		TConstArrayView<FTensorShape> InInputShapes,
+		TConstArrayView<FMLTensorBinding> InOutputTensors)
 	{
 		return -1;
 	}
@@ -138,8 +138,8 @@ protected:
 
 	TArray<FTensorShape>		InputTensorShapes;
 	TArray<FTensorShape>		OutputTensorShapes;
-	TArray<FSymbolicTensorDesc>	InputSymbolicTensors;
-	TArray<FSymbolicTensorDesc>	OutputSymbolicTensors;
+	TArray<FTensorDesc>	InputSymbolicTensors;
+	TArray<FTensorDesc>	OutputSymbolicTensors;
 	EMLInferenceModelType	Type;
 };
 
@@ -148,22 +148,22 @@ protected:
 // IMPLEMENTATION
 //-----------------------------------------------------------------------------
 
-inline TConstArrayView<const FSymbolicTensorDesc> FMLInferenceModel::GetInputTensorDescs() const
+inline TConstArrayView<FTensorDesc> FMLInferenceModel::GetInputTensorDescs() const
 {
 	return InputSymbolicTensors;
 }
 
-inline TConstArrayView<const FSymbolicTensorDesc> FMLInferenceModel::GetOutputTensorDescs() const
+inline TConstArrayView<FTensorDesc> FMLInferenceModel::GetOutputTensorDescs() const
 {
 	return OutputSymbolicTensors;
 }
 
-inline TConstArrayView<const FTensorShape> FMLInferenceModel::GetInputTensorShapes() const
+inline TConstArrayView<FTensorShape> FMLInferenceModel::GetInputTensorShapes() const
 {
 	return InputTensorShapes;
 }
 
-inline TConstArrayView<const FTensorShape> FMLInferenceModel::GetOutputTensorShapes() const
+inline TConstArrayView<FTensorShape> FMLInferenceModel::GetOutputTensorShapes() const
 {
 	return OutputTensorShapes;
 }
diff --git a/Engine/Plugins/Experimental/NNX/Source/NNXCore/Public/NNXTypes.h b/Engine/Plugins/Experimental/NNX/Source/NNXCore/Public/NNXTypes.h
--- a/Engine/Plugins/Experimental/NNX/Source/NNXCore/Public/NNXTypes.h
+++ b/Engine/Plugins/Experimental/NNX/Source/NNXCore/Public/NNXTypes.h
@@ -37,12 +37,12 @@ template <typename DimType, class FinalType> struct TTensorShapeBase
 	constexpr static int32	MaxRank = 8;
 	TArray<DimType, TInlineAllocator<MaxRank>> Data;
 
-	int32 Num() const
+	int32 Rank() const
 	{
 		return Data.Num();
 	}
 
-	static FinalType Make(TConstArrayView<const DimType> Data)
+	static FinalType Make(TConstArrayView<DimType> Data)
 	{
 		if (Data.Num() > MaxRank)
 		{
@@ -61,7 +61,7 @@ struct FSymbolicTensorShape : TTensorShapeBase<int32, FSymbolicTensorShape>
 {
 	bool IsConcrete() const
 	{
-		for (int32 i = 0; i < Data.Num(); ++i)
+		for (int32 i = 0; i < Rank(); ++i)
 		{
 			if (Data[i] < 0)
 			{
@@ -88,11 +88,11 @@ struct FTensorShape : TTensorShapeBase<uint32, FTensorShape>
 	
 	bool IsCompatibleWith(FSymbolicTensorShape SymbolicShape) const
 	{
-		if (Num() != SymbolicShape.Num())
+		if (Rank() != SymbolicShape.Rank())
 		{
 			return false;
 		}
-		for (int32 i = 0; i < Num(); ++i)
+		for (int32 i = 0; i < Rank(); ++i)
 		{
 			if (SymbolicShape.Data[i] >= 0 && SymbolicShape.Data[i] != Data[i])
 			{
@@ -151,11 +151,25 @@ inline size_t GetTensorDataTypeSizeInBytes(EMLTensorDataType InType)
 	return 0;
 }
 
-struct FTensorDescBase
+class FTensorDescBase
 {
+protected:
 	FString					Name;
 	EMLTensorDataType		DataType;
 
+	FTensorDescBase() = default;
+	
+public:
+	const FString& GetName() const
+	{
+		return Name;
+	}
+	
+	EMLTensorDataType GetDataType() const
+	{
+		return DataType;
+	}
+
 	/** Return size of one element in bytes */
 	uint32 GetElemByteSize() const
 	{
@@ -168,14 +182,24 @@ struct FTensorDescBase
 	}
 };
 
-/** Symbolic tensor descriptor */
-struct FSymbolicTensorDesc : FTensorDescBase
+/** Symbolic tensor descriptor without data */
+class FTensorDesc : public FTensorDescBase
 {
 	FSymbolicTensorShape	Shape;
 
-	static FSymbolicTensorDesc Make(const FString& Name, const FSymbolicTensorShape& Shape, EMLTensorDataType DataType)
+public:
+	const FSymbolicTensorShape& GetShape() const
+	{
+		return Shape;
+	}
+
+	static FTensorDesc Make(const FString& Name, const FSymbolicTensorShape& Shape, EMLTensorDataType DataType)
 	{
-		return { {Name, DataType} , Shape };
+		FTensorDesc Desc;
+		Desc.Name = Name;
+		Desc.DataType = DataType;
+		Desc.Shape = Shape;
+		return Desc;
 	}
 
 	bool IsConcrete() const
@@ -184,25 +208,44 @@ struct FSymbolicTensorDesc : FTensorDescBase
 	}
 };
 
-/** Concrete tensor descriptor */
- struct FMLTensorDesc : FTensorDescBase
+/** Concrete tensor backed by data */
+class FTensor : public FTensorDescBase
 {
 	FTensorShape Shape;
-	uint32 Volume;
-	uint64 DataSize;
+	uint32 Volume = 0;
+	uint64 DataSize = 0;
 
-	static FMLTensorDesc Make(const FString& Name, const FTensorShape& Shape, EMLTensorDataType DataType)
+public:
+	const FTensorShape& GetShape() const
 	{
-		uint64 Volume = Shape.Volume();
-		uint64 DataSize = (uint64)GetTensorDataTypeSizeInBytes(DataType) * Volume;
-		
-		check(Volume <= TNumericLimits<uint32>::Max());
-		return { {Name, DataType}, Shape, (uint32)Volume, DataSize };
+		return Shape;
+	}
+
+	uint32 GetVolume() const
+	{
+		return Volume;
+	}
+
+	uint64 GetDataSize() const
+	{
+		return DataSize;
+	}
+	
+	static FTensor Make(const FString& Name, const FTensorShape& Shape, EMLTensorDataType DataType)
+	{
+		FTensor Desc;
+		Desc.Name = Name;
+		Desc.DataType = DataType;
+		Desc.Shape = Shape;
+		Desc.Volume = Shape.Volume();
+		Desc.DataSize = (uint64)GetTensorDataTypeSizeInBytes(DataType) * Desc.Volume;
+		check(Desc.Volume <= TNumericLimits<uint32>::Max());
+		return Desc;
 	}
 
-	static FMLTensorDesc MakeFromSymbolic(const FSymbolicTensorDesc& TensorDesc)
+	static FTensor MakeFromSymbolicDesc(const FTensorDesc& TensorDesc)
 	{
-		return Make(TensorDesc.Name, FTensorShape::MakeFromSymbolic(TensorDesc.Shape), TensorDesc.DataType);
+		return Make(TensorDesc.GetName(), FTensorShape::MakeFromSymbolic(TensorDesc.GetShape()), TensorDesc.GetDataType());
 	}
 };
 
diff --git a/Engine/Plugins/Experimental/NNX/Source/NNXQA/Private/NNXQAParametricTest.cpp b/Engine/Plugins/Experimental/NNX/Source/NNXQA/Private/NNXQAParametricTest.cpp
--- a/Engine/Plugins/Experimental/NNX/Source/NNXQA/Private/NNXQAParametricTest.cpp
+++ b/Engine/Plugins/Experimental/NNX/Source/NNXQA/Private/NNXQAParametricTest.cpp
@@ -147,7 +147,7 @@ namespace Test
 			{
 				FTensorShape Shape = GetConcreteShapeFromJsonArray(Tensor.Shape);
 				EMLTensorDataType TensorType = GetTensorTypeFromJson(Tensor.Type, DefaultInputType);
-				TestSetup.Inputs.Emplace(FMLTensorDesc::Make(FString::Printf(TEXT("in%d"), i++), Shape, TensorType));
+				TestSetup.Inputs.Emplace(FTensor::Make(FString::Printf(TEXT("in%d"), i++), Shape, TensorType));
 			}
 
 			i = 0;
@@ -155,14 +155,14 @@ namespace Test
 			{
 				FTensorShape Shape = GetConcreteShapeFromJsonArray(Tensor.Shape);
 				EMLTensorDataType TensorType = GetTensorTypeFromJson(Tensor.Type, DefaultOutputType);
-				TestSetup.Outputs.Emplace(FMLTensorDesc::Make(FString::Printf(TEXT("output%d"), i++), Shape, TensorType));
+				TestSetup.Outputs.Emplace(FTensor::Make(FString::Printf(TEXT("output%d"), i++), Shape, TensorType));
 			}
 			//If output is not defined it is the first input shape.
 			if (TestDataset.Outputs.Num() == 0 && TestDataset.Inputs.Num() > 0)
 			{
 				FTensorShape Shape = GetConcreteShapeFromJsonArray(TestDataset.Inputs[0].Shape);
 				EMLTensorDataType TensorType = GetTensorTypeFromJson(TestDataset.Inputs[0].Type, DefaultOutputType);
-				TestSetup.Outputs.Emplace(FMLTensorDesc::Make(TEXT("output"), Shape, TensorType));
+				TestSetup.Outputs.Emplace(FTensor::Make(TEXT("output"), Shape, TensorType));
 			}
 		}
 
diff --git a/Engine/Plugins/Experimental/NNX/Source/NNXQA/Private/NNXQAUtils.cpp b/Engine/Plugins/Experimental/NNX/Source/NNXQA/Private/NNXQAUtils.cpp
--- a/Engine/Plugins/Experimental/NNX/Source/NNXQA/Private/NNXQAUtils.cpp
+++ b/Engine/Plugins/Experimental/NNX/Source/NNXQA/Private/NNXQAUtils.cpp
@@ -14,7 +14,7 @@ namespace NNX
 {
 namespace Test 
 {
-	static void FillTensorBindings(TConstArrayView<const FMLTensorDesc> TensorDescs,
+	static void FillTensorBindings(TConstArrayView<FTensor> TensorDescs,
 		TArray<TArray<char>>& OutMemBuffers, TArray<FMLTensorBinding>& OutBindings,
         std::function<float(EMLTensorDataType, uint32, uint32)> Initializer)
 	{
@@ -23,11 +23,11 @@ namespace Test
 
 		for (int Index = 0; Index < TensorDescs.Num(); ++Index)
 		{
-			const FMLTensorDesc& TensorDesc = TensorDescs[Index];
-			const uint32 NumberOfElements = TensorDesc.Volume;
+			const FTensor& TensorDesc = TensorDescs[Index];
+			const uint32 NumberOfElements = TensorDesc.GetVolume();
 			const int32 ElementByteSize = TensorDesc.GetElemByteSize();
-			const uint64 BufferSize = TensorDesc.DataSize;
-			const EMLTensorDataType DataType = TensorDesc.DataType;
+			const uint64 BufferSize = TensorDesc.GetDataSize();
+			const EMLTensorDataType DataType = TensorDesc.GetDataType();
 
 			TArray<char> CurInputBuffer;
 	        CurInputBuffer.SetNum(BufferSize);
@@ -85,21 +85,21 @@ namespace Test
 	template FString ShapeToString<int32>(TArrayView<const int32> Shape);
 	template FString ShapeToString<uint32>(TArrayView<const uint32> Shape);
 
-	FString FMLTensorDescToString(const FMLTensorDesc& desc)
+	FString FMLTensorDescToString(const FTensor& desc)
 	{
-		TArrayView<const uint32> Shape(desc.Shape.Data);
+		TArrayView<const uint32> Shape(desc.GetShape().Data);
 
 		FString TensorDesc;
 		TensorDesc.Reserve(50);
-		TensorDesc += FString::Printf(TEXT("Name: %s, Shape: "), *desc.Name);
+		TensorDesc += FString::Printf(TEXT("Name: %s, Shape: "), *desc.GetName());
 		TensorDesc += ShapeToString(Shape);
-		const FString& DataTypeName = StaticEnum<EMLTensorDataType>()->GetNameStringByValue(static_cast<int64>(desc.DataType));
+		const FString& DataTypeName = StaticEnum<EMLTensorDataType>()->GetNameStringByValue(static_cast<int64>(desc.GetDataType()));
 		TensorDesc += FString::Printf(TEXT(" DataType: %s"), *DataTypeName);
 
 		return TensorDesc;
 	}
 
-	FString TensorToString(const FMLTensorDesc& TensorDesc, const TArray<char>& TensorData)
+	FString TensorToString(const FTensor& TensorDesc, const TArray<char>& TensorData)
 	{
 		FString TensorLog;
 		TensorLog.Reserve(50);
@@ -108,7 +108,7 @@ namespace Test
 		TensorLog += TEXT(", Data: ");
 
 		const constexpr uint32 MAX_DATA_TO_LOG = 10;
-		const uint32 MaxIndex = FMath::Min(MAX_DATA_TO_LOG, TensorDesc.Volume);
+		const uint32 MaxIndex = FMath::Min(MAX_DATA_TO_LOG, TensorDesc.GetVolume());
         const uint32 ElementByteSize = TensorDesc.GetElemByteSize();
 		for (uint32 i = 0; i < MaxIndex; ++i)
 		{
@@ -116,7 +116,7 @@ namespace Test
 			const char* Data = TensorData.GetData() + ByteOffset;
 			check((int32)ByteOffset <= TensorData.Num());
 
-			switch (TensorDesc.DataType)
+			switch (TensorDesc.GetDataType())
 			{
 				case EMLTensorDataType::Float:
 					TensorLog += FString::Printf(TEXT("%0.2f"), *(float*)Data); break;
@@ -133,24 +133,24 @@ namespace Test
 			if (i < MaxIndex)
 				TensorLog += TEXT(", ");
 		}
-		if (MaxIndex < TensorDesc.Volume)
+		if (MaxIndex < TensorDesc.GetVolume())
 			TensorLog += TEXT(",...");
 
 		return TensorLog;
 	}
 
 	template<typename T> bool CompareTensorData(
-		const NNX::FMLTensorDesc& RefTensorDesc,   const TArray<char>& RefRawBuffer,
-		const NNX::FMLTensorDesc& OtherTensorDesc, const TArray<char>& OtherRawBuffer,
+		const NNX::FTensor& RefTensorDesc,   const TArray<char>& RefRawBuffer,
+		const NNX::FTensor& OtherTensorDesc, const TArray<char>& OtherRawBuffer,
 		float AbsoluteErrorEpsilon, float RelativeErrorPercent)
 	{
 		const T* RefBuffer = (T*)RefRawBuffer.GetData();
 		const T* OtherBuffer = (T*)OtherRawBuffer.GetData();
 		
-		const uint32 Volume = RefTensorDesc.Volume;
+		const uint32 Volume = RefTensorDesc.GetVolume();
 		const uint32 ElementByteSize = RefTensorDesc.GetElemByteSize();
 
-		check(Volume == OtherTensorDesc.Volume);
+		check(Volume == OtherTensorDesc.GetVolume());
 		check(Volume * ElementByteSize == RefRawBuffer.Num());
 		check(Volume * ElementByteSize == OtherRawBuffer.Num());
 
@@ -242,18 +242,18 @@ namespace Test
 	}
 
 	bool VerifyTensorResult(
-		const NNX::FMLTensorDesc& RefTensorDesc, const TArray<char>& RefRawBuffer,
-		const NNX::FMLTensorDesc& OtherTensorDesc, const TArray<char>& OtherRawBuffer,
+		const NNX::FTensor& RefTensorDesc, const TArray<char>& RefRawBuffer,
+		const NNX::FTensor& OtherTensorDesc, const TArray<char>& OtherRawBuffer,
 		float AbsoluteErrorEpsilon, float RelativeErrorPercent)
 	{
 		bool bTensorDescMatch = true;
-		bTensorDescMatch &= (RefTensorDesc.Name == OtherTensorDesc.Name);
-		bTensorDescMatch &= (RefTensorDesc.Shape.Num() == OtherTensorDesc.Shape.Num());
-		bTensorDescMatch &= (RefTensorDesc.DataType == OtherTensorDesc.DataType);
-		check(RefTensorDesc.Shape.Num() <= FTensorShape::MaxRank);
-		for (int32 i = 0; i < RefTensorDesc.Shape.Num(); ++i)
+		bTensorDescMatch &= (RefTensorDesc.GetName() == OtherTensorDesc.GetName());
+		bTensorDescMatch &= (RefTensorDesc.GetShape().Rank() == OtherTensorDesc.GetShape().Rank());
+		bTensorDescMatch &= (RefTensorDesc.GetDataType() == OtherTensorDesc.GetDataType());
+		check(RefTensorDesc.GetShape().Rank() <= FTensorShape::MaxRank);
+		for (int32 i = 0; i < RefTensorDesc.GetShape().Rank(); ++i)
 		{
-			bTensorDescMatch &= (RefTensorDesc.Shape.Data[i] == OtherTensorDesc.Shape.Data[i]);
+			bTensorDescMatch &= (RefTensorDesc.GetShape().Data[i] == OtherTensorDesc.GetShape().Data[i]);
 		}
 
 		if (!bTensorDescMatch)
@@ -262,20 +262,20 @@ namespace Test
 			return false;
 		}
 
-		if (RefTensorDesc.DataType == EMLTensorDataType::Float)
+		if (RefTensorDesc.GetDataType() == EMLTensorDataType::Float)
 		{
 			return CompareTensorData<float>(RefTensorDesc, RefRawBuffer, OtherTensorDesc, OtherRawBuffer, AbsoluteErrorEpsilon, RelativeErrorPercent);
 		}
-		else if (RefTensorDesc.DataType == EMLTensorDataType::Boolean)
+		else if (RefTensorDesc.GetDataType() == EMLTensorDataType::Boolean)
 		{
 			check(RefTensorDesc.GetElemByteSize() == 1);
 			return CompareTensorData<bool>(RefTensorDesc, RefRawBuffer, OtherTensorDesc, OtherRawBuffer, AbsoluteErrorEpsilon, RelativeErrorPercent);
 		}
-		else if (RefTensorDesc.DataType == EMLTensorDataType::Int32)
+		else if (RefTensorDesc.GetDataType() == EMLTensorDataType::Int32)
 		{
 			return CompareTensorData<int32>(RefTensorDesc, RefRawBuffer, OtherTensorDesc, OtherRawBuffer, AbsoluteErrorEpsilon, RelativeErrorPercent);
 		}
-		else if (RefTensorDesc.DataType == EMLTensorDataType::UInt32)
+		else if (RefTensorDesc.GetDataType() == EMLTensorDataType::UInt32)
 		{
 			return CompareTensorData<uint32>(RefTensorDesc, RefRawBuffer, OtherTensorDesc, OtherRawBuffer, AbsoluteErrorEpsilon, RelativeErrorPercent);
 		}
@@ -326,7 +326,7 @@ namespace Test
 	}
 
 	static int RunTestInference(const FNNIModelRaw& ONNXModelData, IRuntime* Runtime, 
-		TArray<FMLTensorDesc>& OutOutputTensorsDesc, TArray<TArray<char>>& OutOutputMemBuffers)
+		TArray<FTensor>& OutOutputTensorsDesc, TArray<TArray<char>>& OutOutputMemBuffers)
 	{
 		OutOutputMemBuffers.Empty();
 		OutOutputTensorsDesc.Empty();
@@ -351,28 +351,28 @@ namespace Test
 		}
 
 		//bind input tensors to memory (CPU) and initialize
-		TConstArrayView<const FSymbolicTensorDesc> InputSymbolicTensorDescs = InferenceModel->GetInputTensorDescs();
+		TConstArrayView<FTensorDesc> InputSymbolicTensorDescs = InferenceModel->GetInputTensorDescs();
 		TArray<FTensorShape> InputShapes;
-		TArray<FMLTensorDesc> InputTensorsDescs;
+		TArray<FTensor> InputTensorsDescs;
 		TArray<FMLTensorBinding> InputBindings;
 		TArray<TArray<char>> InputMemBuffers;
 		
-		for (const FSymbolicTensorDesc& SymbolicTensorDesc : InputSymbolicTensorDescs)
+		for (const FTensorDesc& SymbolicTensorDesc : InputSymbolicTensorDescs)
 		{
-			FMLTensorDesc TensorDesc = FMLTensorDesc::MakeFromSymbolic(SymbolicTensorDesc);
-			InputShapes.Emplace(TensorDesc.Shape);
+			FTensor TensorDesc = FTensor::MakeFromSymbolicDesc(SymbolicTensorDesc);
+			InputShapes.Emplace(TensorDesc.GetShape());
 			InputTensorsDescs.Emplace(TensorDesc);
 		}
 		FillTensorBindings(InputTensorsDescs, InputMemBuffers, InputBindings, InputTensorInitializer);
 		
 		//bind output tensors to memory (CPU) and initialize
-		TConstArrayView<const FSymbolicTensorDesc> OutputSymbolicTensorDescs = InferenceModel->GetOutputTensorDescs();
-		TArray<FMLTensorDesc> OutputTensorsDescs;
+		TConstArrayView<FTensorDesc> OutputSymbolicTensorDescs = InferenceModel->GetOutputTensorDescs();
+		TArray<FTensor> OutputTensorsDescs;
 		TArray<NNX::FMLTensorBinding> OutputBindings;
 
-		for (const FSymbolicTensorDesc& SymbolicTensorDesc : OutputSymbolicTensorDescs)
+		for (const FTensorDesc& SymbolicTensorDesc : OutputSymbolicTensorDescs)
 		{
-			FMLTensorDesc TensorDesc = FMLTensorDesc::MakeFromSymbolic(SymbolicTensorDesc);
+			FTensor TensorDesc = FTensor::MakeFromSymbolicDesc(SymbolicTensorDesc);
 			OutputTensorsDescs.Emplace(TensorDesc);
 		}
 		OutOutputTensorsDesc = OutputTensorsDescs;
@@ -421,7 +421,7 @@ namespace Test
 		}
 		const FString& RefName = RefRuntime->GetRuntimeName();
 		TArray<TArray<char>> RefOutputMemBuffers;
-		TArray<NNX::FMLTensorDesc> RefOutputTensorDescs;
+		TArray<NNX::FTensor> RefOutputTensorDescs;
 
 		RunTestInference(ONNXModel, RefRuntime, RefOutputTensorDescs, RefOutputMemBuffers);
 
@@ -459,7 +459,7 @@ namespace Test
 			else
 			{
 				TArray<TArray<char>> OutputMemBuffers;
-				TArray<NNX::FMLTensorDesc> OutputTensorDescs;
+				TArray<NNX::FTensor> OutputTensorDescs;
 				bool bTestSuceeded = true;
 				float AbsoluteErrorEpsilon = TestSetup.GetAbsoluteErrorEpsilonForRuntime(RuntimeName);
 				float RelativeErrorPercent = TestSetup.GetRelativeErrorPercentForRuntime(RuntimeName);
diff --git a/Engine/Plugins/Experimental/NNX/Source/NNXQA/Private/NNXQAUtils.h b/Engine/Plugins/Experimental/NNX/Source/NNXQA/Private/NNXQAUtils.h
--- a/Engine/Plugins/Experimental/NNX/Source/NNXQA/Private/NNXQAUtils.h
+++ b/Engine/Plugins/Experimental/NNX/Source/NNXQA/Private/NNXQAUtils.h
@@ -34,8 +34,8 @@ namespace Test
 			bool IsModelTest;
 			TMap<FString, float> AbsoluteErrorEpsilonForRuntime;
 			TMap<FString, float> RelativeErrorPercentForRuntime;
-			TArray<FMLTensorDesc> Inputs;
-			TArray<FMLTensorDesc> Outputs;
+			TArray<FTensor> Inputs;
+			TArray<FTensor> Outputs;
 			UE::NNECore::FAttributeMap AttributeMap;
 			TArray<FString> Tags;
 			TArray<FString> AutomationExcludedRuntime;
@@ -63,8 +63,8 @@ namespace Test
 
 	bool CompareONNXModelInferenceAcrossRuntimes(const FNNIModelRaw& ONNXModel, const FTests::FTestSetup& TestSetup, const FString& RuntimeFilter = TEXT(""));
 
-	FString FMLTensorDescToString(const FMLTensorDesc& TensorDesc);
-	FString TensorToString(const FMLTensorDesc& TensorDesc, const TArray<char>& TensorData);
+	FString FMLTensorDescToString(const FTensor& TensorDesc);
+	FString TensorToString(const FTensor& TensorDesc, const TArray<char>& TensorData);
 	template<typename T> FString ShapeToString(TArrayView<const T> Shape);
 
 } // namespace Test
diff --git a/Engine/Plugins/Experimental/NNX/Source/NNXRuntimeCPU/Private/NNXRuntimeCPU.cpp b/Engine/Plugins/Experimental/NNX/Source/NNXRuntimeCPU/Private/NNXRuntimeCPU.cpp
--- a/Engine/Plugins/Experimental/NNX/Source/NNXRuntimeCPU/Private/NNXRuntimeCPU.cpp
+++ b/Engine/Plugins/Experimental/NNX/Source/NNXRuntimeCPU/Private/NNXRuntimeCPU.cpp
@@ -163,18 +163,15 @@ bool FMLInferenceModelCPU::ConfigureTensors(const bool InIsInput)
 
 	const uint32 NumberTensors = bIsInput ? Session->GetInputCount() : Session->GetOutputCount();
 
-	TArray<FSymbolicTensorDesc>& SymbolicTensorDescs = bIsInput ? InputSymbolicTensors : OutputSymbolicTensors;
+	TArray<FTensorDesc>& SymbolicTensorDescs = bIsInput ? InputSymbolicTensors : OutputSymbolicTensors;
 	TArray<ONNXTensorElementDataType>& TensorsORTType = bIsInput ? InputTensorsORTType : OutputTensorsORTType;
 	TArray<const char*>& TensorNames = bIsInput ? InputTensorNames : OutputTensorNames;
 
 	for (uint32 TensorIndex = 0; TensorIndex < NumberTensors; ++TensorIndex)
 	{
-		FSymbolicTensorDesc SymbolicTensorDesc;
-
 		// Get Tensor name
 		const char* CurTensorName = bIsInput ? Session->GetInputName(TensorIndex, *Allocator) : Session->GetOutputName(TensorIndex, *Allocator);
 		TensorNames.Emplace(CurTensorName);
-		SymbolicTensorDesc.Name = FString(CurTensorName);
 
 		// Get node type
 		Ort::TypeInfo CurrentTypeInfo = bIsInput ? Session->GetInputTypeInfo(TensorIndex) : Session->GetOutputTypeInfo(TensorIndex);
@@ -185,14 +182,15 @@ bool FMLInferenceModelCPU::ConfigureTensors(const bool InIsInput)
 		TensorsORTType.Emplace(ONNXTensorElementDataTypeEnum);
 
 		std::pair<EMLTensorDataType, uint64> TypeAndSize = TranslateTensorTypeORTToNNI(ONNXTensorElementDataTypeEnum);
-		SymbolicTensorDesc.DataType = TypeAndSize.first;
 
-		SymbolicTensorDesc.Shape.Data.Reserve(CurrentTensorInfo.GetShape().size());
+		FSymbolicTensorShape Shape;
+		Shape.Data.Reserve(CurrentTensorInfo.GetShape().size());
 		for (const int64_t CurrentTensorSize : CurrentTensorInfo.GetShape())
 		{
-			SymbolicTensorDesc.Shape.Data.Add((int32)CurrentTensorSize);
+			Shape.Data.Add((int32)CurrentTensorSize);
 		}
 
+		FTensorDesc SymbolicTensorDesc = FTensorDesc::Make(FString(CurTensorName), Shape, TypeAndSize.first);
 		if (!SymbolicTensorDesc.IsConcrete())
 			{
 				UE_LOG(LogNNX, Display,
@@ -207,33 +205,33 @@ bool FMLInferenceModelCPU::ConfigureTensors(const bool InIsInput)
 	return true;
 }
 
-int FMLInferenceModelCPU::SetInputShapes(TConstArrayView<const FTensorShape> InInputShapes)
-			{
+int FMLInferenceModelCPU::SetInputShapes(TConstArrayView<FTensorShape> InInputShapes)
+{
 	//Verify input shape are valid for the model
 	if (FMLInferenceModel::SetInputShapes(InInputShapes) != 0)
 	{
 		return -1;
-		}
+	}
 
 	//Run shape inference
 	//TODO jira 168972: handle dynamic input shape
 
 	//Build concrete input and output tensor descriptions
 	InputTensors.Empty();
-	for (FSymbolicTensorDesc SymbolicTensorDesc : InputSymbolicTensors)
+	for (FTensorDesc SymbolicTensorDesc : InputSymbolicTensors)
 	{
-		InputTensors.Emplace(FMLTensorDesc::MakeFromSymbolic(SymbolicTensorDesc));
+		InputTensors.Emplace(FTensor::MakeFromSymbolicDesc(SymbolicTensorDesc));
 	}
 	OutputTensors.Empty();
-	for (FSymbolicTensorDesc SymbolicTensorDesc : OutputSymbolicTensors)
+	for (FTensorDesc SymbolicTensorDesc : OutputSymbolicTensors)
 	{
-		OutputTensors.Emplace(FMLTensorDesc::MakeFromSymbolic(SymbolicTensorDesc));
-}
+		OutputTensors.Emplace(FTensor::MakeFromSymbolicDesc(SymbolicTensorDesc));
+	}
 
 	return 0;
 }
 
-int FMLInferenceModelCPU::Run(TConstArrayView<const FMLTensorBinding> InInputBindings, TConstArrayView<const FTensorShape> InInputShapes, TConstArrayView<const FMLTensorBinding> InOutputBindings)
+int FMLInferenceModelCPU::Run(TConstArrayView<FMLTensorBinding> InInputBindings, TConstArrayView<FTensorShape> InInputShapes, TConstArrayView<FMLTensorBinding> InOutputBindings)
 {
 	DECLARE_SCOPE_CYCLE_COUNTER(TEXT("FMLInferenceModelCPU_Run"), STAT_FMLInferenceModelCPU_Run, STATGROUP_MachineLearning);
 
diff --git a/Engine/Plugins/Experimental/NNX/Source/NNXRuntimeCPU/Public/NNXRuntimeCPU.h b/Engine/Plugins/Experimental/NNX/Source/NNXRuntimeCPU/Public/NNXRuntimeCPU.h
--- a/Engine/Plugins/Experimental/NNX/Source/NNXRuntimeCPU/Public/NNXRuntimeCPU.h
+++ b/Engine/Plugins/Experimental/NNX/Source/NNXRuntimeCPU/Public/NNXRuntimeCPU.h
@@ -84,8 +84,8 @@ namespace NNX
 		bool Init(const UMLInferenceModel* InferenceModel);
 		bool IsLoaded() const;
 
-		virtual int SetInputShapes(TConstArrayView<const FTensorShape> InInputShapes) override;
-		virtual int Run(TConstArrayView<const FMLTensorBinding> InInputBindings, TConstArrayView<const FTensorShape> InInputShapes, TConstArrayView<const FMLTensorBinding> InOutputBindings) override;
+		virtual int SetInputShapes(TConstArrayView<FTensorShape> InInputShapes) override;
+		virtual int Run(TConstArrayView<FMLTensorBinding> InInputBindings, TConstArrayView<FTensorShape> InInputShapes, TConstArrayView<FMLTensorBinding> InOutputBindings) override;
 
 		float GetLastRunTimeMSec() const;
 		UE::NNEProfiling::Internal::FStatistics GetRunStatistics() const;
@@ -114,8 +114,8 @@ namespace NNX
 		TArray<const char*> InputTensorNames;
 		TArray<const char*> OutputTensorNames;
 
-		TArray<FMLTensorDesc> InputTensors;
-		TArray<FMLTensorDesc> OutputTensors;
+		TArray<FTensor> InputTensors;
+		TArray<FTensor> OutputTensors;
 
 		bool InitializedAndConfigureMembers();
 		bool ConfigureTensors(const bool InIsInput);
diff --git a/Engine/Plugins/Experimental/NNX/Source/NNXRuntimeCPU/Public/NNXRuntimeCPUUtils.h b/Engine/Plugins/Experimental/NNX/Source/NNXRuntimeCPU/Public/NNXRuntimeCPUUtils.h
--- a/Engine/Plugins/Experimental/NNX/Source/NNXRuntimeCPU/Public/NNXRuntimeCPUUtils.h
+++ b/Engine/Plugins/Experimental/NNX/Source/NNXRuntimeCPU/Public/NNXRuntimeCPUUtils.h
@@ -113,7 +113,7 @@ inline TypeInfoORT TranslateTensorTypeORTToNNI(unsigned int OrtDataType)
 
 inline void BindTensorsToORT(
 	TArrayView<const NNX::FMLTensorBinding> InBindingTensors,
-	const TArray<NNX::FMLTensorDesc>& InTensorsDescriptors,
+	const TArray<NNX::FTensor>& InTensorsDescriptors,
 	const TArray<ONNXTensorElementDataType>& InTensorsORTType, 
 	const Ort::MemoryInfo* InAllocatorInfo,
 	TArray<Ort::Value>& OutOrtTensors	
@@ -131,18 +131,18 @@ inline void BindTensorsToORT(
 	for (uint32 Index = 0; Index < NumBinding; ++Index)
 	{
 		const NNX::FMLTensorBinding& CurrentBinding = InBindingTensors[Index];
-		const NNX::FMLTensorDesc& CurrentDescriptor = InTensorsDescriptors[Index];
+		const NNX::FTensor& CurrentDescriptor = InTensorsDescriptors[Index];
 		ONNXTensorElementDataType CurrentORTType = InTensorsORTType[Index];
 
 		TUniquePtr<int64_t[]> SizesInt64t;
-		SizesInt64t = MakeUnique<int64_t[]>(CurrentDescriptor.Shape.Num());
-		for (int32 DimIndex = 0; DimIndex < CurrentDescriptor.Shape.Num(); ++DimIndex)
+		SizesInt64t = MakeUnique<int64_t[]>(CurrentDescriptor.GetShape().Rank());
+		for (int32 DimIndex = 0; DimIndex < CurrentDescriptor.GetShape().Rank(); ++DimIndex)
 		{
-			SizesInt64t.Get()[DimIndex] = CurrentDescriptor.Shape.Data[DimIndex];
+			SizesInt64t.Get()[DimIndex] = CurrentDescriptor.GetShape().Data[DimIndex];
 		}
 
-		const uint64 ByteCount { InTensorsDescriptors[Index].DataSize };
-		const uint32 ArrayDimensions { (uint32)CurrentDescriptor.Shape.Num() };
+		const uint64 ByteCount { InTensorsDescriptors[Index].GetDataSize() };
+		const uint32 ArrayDimensions { (uint32)CurrentDescriptor.GetShape().Rank() };
 		OutOrtTensors.Emplace(
 			Ort::Value::CreateTensor(
 				*InAllocatorInfo,
diff --git a/Engine/Plugins/Experimental/NNX/Source/NNXRuntimeORT/Private/NNXRuntimeORT.cpp b/Engine/Plugins/Experimental/NNX/Source/NNXRuntimeORT/Private/NNXRuntimeORT.cpp
--- a/Engine/Plugins/Experimental/NNX/Source/NNXRuntimeORT/Private/NNXRuntimeORT.cpp
+++ b/Engine/Plugins/Experimental/NNX/Source/NNXRuntimeORT/Private/NNXRuntimeORT.cpp
@@ -237,18 +237,15 @@ bool FMLInferenceModelORT::ConfigureTensors(const bool InIsInput)
 
 	const uint32 NumberTensors = bIsInput ? Session->GetInputCount() : Session->GetOutputCount();
 
-	TArray<FSymbolicTensorDesc>& SymbolicTensorDescs = bIsInput ? InputSymbolicTensors : OutputSymbolicTensors;
+	TArray<FTensorDesc>& SymbolicTensorDescs = bIsInput ? InputSymbolicTensors : OutputSymbolicTensors;
 	TArray<ONNXTensorElementDataType>& TensorsORTType = bIsInput ? InputTensorsORTType : OutputTensorsORTType;
 	TArray<const char*>& TensorNames = bIsInput ? InputTensorNames : OutputTensorNames;
 
 	for (uint32 TensorIndex = 0; TensorIndex < NumberTensors; ++TensorIndex)
 	{
-		FSymbolicTensorDesc SymbolicTensorDesc;
-
 		// Get Tensor name
 		const char* CurTensorName = bIsInput ? Session->GetInputName(TensorIndex, *Allocator) : Session->GetOutputName(TensorIndex, *Allocator);
 		TensorNames.Emplace(CurTensorName);
-		SymbolicTensorDesc.Name = FString(CurTensorName);
 
 		// Get node type
 		Ort::TypeInfo CurrentTypeInfo = bIsInput ? Session->GetInputTypeInfo(TensorIndex) : Session->GetOutputTypeInfo(TensorIndex);
@@ -259,14 +256,16 @@ bool FMLInferenceModelORT::ConfigureTensors(const bool InIsInput)
 		TensorsORTType.Emplace(ONNXTensorElementDataTypeEnum);
 
 		std::pair<EMLTensorDataType, uint64> TypeAndSize = TranslateTensorTypeORTToNNI(ONNXTensorElementDataTypeEnum);
-		SymbolicTensorDesc.DataType = TypeAndSize.first;
 		
-		SymbolicTensorDesc.Shape.Data.Reserve(CurrentTensorInfo.GetShape().size());
+		FSymbolicTensorShape Shape;
+		Shape.Data.Reserve(CurrentTensorInfo.GetShape().size());
 		for (const int64_t CurrentTensorSize : CurrentTensorInfo.GetShape())
 		{
-			SymbolicTensorDesc.Shape.Data.Add((int32)CurrentTensorSize);
+			Shape.Data.Add((int32)CurrentTensorSize);
 		}
 
+		FTensorDesc SymbolicTensorDesc = FTensorDesc::Make(FString(CurTensorName), Shape, TypeAndSize.first);
+
 		if (!SymbolicTensorDesc.IsConcrete())
 			{
 				UE_LOG(LogNNX, Display,
@@ -281,33 +280,33 @@ bool FMLInferenceModelORT::ConfigureTensors(const bool InIsInput)
 	return true;
 }
 
-int FMLInferenceModelORT::SetInputShapes(TConstArrayView<const FTensorShape> InInputShapes)
+int FMLInferenceModelORT::SetInputShapes(TConstArrayView<FTensorShape> InInputShapes)
 {
 	//Verify input shape are valid for the model
 	if (FMLInferenceModel::SetInputShapes(InInputShapes) != 0)
-			{
+	{
 		return -1;
-			}
+	}
 
 	//Run shape inference
 	//TODO jira 168972: handle dynamic input shape
 
 	//Build concrete input and output tensor descriptions
 	InputTensors.Empty();
-	for (FSymbolicTensorDesc SymbolicTensorDesc : InputSymbolicTensors)
+	for (FTensorDesc SymbolicTensorDesc : InputSymbolicTensors)
 	{
-		InputTensors.Emplace(FMLTensorDesc::MakeFromSymbolic(SymbolicTensorDesc));
-		}
+		InputTensors.Emplace(FTensor::MakeFromSymbolicDesc(SymbolicTensorDesc));
+	}
 	OutputTensors.Empty();
-	for (FSymbolicTensorDesc SymbolicTensorDesc : OutputSymbolicTensors)
+	for (FTensorDesc SymbolicTensorDesc : OutputSymbolicTensors)
 	{
-		OutputTensors.Emplace(FMLTensorDesc::MakeFromSymbolic(SymbolicTensorDesc));
+		OutputTensors.Emplace(FTensor::MakeFromSymbolicDesc(SymbolicTensorDesc));
 	}
 
 	return 0;
 }
 
-int FMLInferenceModelORT::Run(TConstArrayView<const FMLTensorBinding> InInputBindings, TConstArrayView<const FTensorShape> InInputShapes, TConstArrayView<const FMLTensorBinding> InOutputBindings)
+int FMLInferenceModelORT::Run(TConstArrayView<FMLTensorBinding> InInputBindings, TConstArrayView<FTensorShape> InInputShapes, TConstArrayView<FMLTensorBinding> InOutputBindings)
 {
 	DECLARE_SCOPE_CYCLE_COUNTER(TEXT("FMLInferenceModelORT_Run"), STAT_FMLInferenceModelORT_Run, STATGROUP_MachineLearning);
 
diff --git a/Engine/Plugins/Experimental/NNX/Source/NNXRuntimeORT/Private/NNXRuntimeORT.h b/Engine/Plugins/Experimental/NNX/Source/NNXRuntimeORT/Private/NNXRuntimeORT.h
--- a/Engine/Plugins/Experimental/NNX/Source/NNXRuntimeORT/Private/NNXRuntimeORT.h
+++ b/Engine/Plugins/Experimental/NNX/Source/NNXRuntimeORT/Private/NNXRuntimeORT.h
@@ -183,8 +183,8 @@ namespace NNX
 		bool Init(const UMLInferenceModel* InferenceModel);
 		bool IsLoaded() const;
 
-		virtual int SetInputShapes(TConstArrayView<const FTensorShape> InInputShapes) override;
-		virtual int Run(TConstArrayView<const FMLTensorBinding> InInputBindings, TConstArrayView<const FTensorShape> InInputShapes, TConstArrayView<const FMLTensorBinding> InOutputBindings) override;
+		virtual int SetInputShapes(TConstArrayView<FTensorShape> InInputShapes) override;
+		virtual int Run(TConstArrayView<FMLTensorBinding> InInputBindings, TConstArrayView<FTensorShape> InInputShapes, TConstArrayView<FMLTensorBinding> InOutputBindings) override;
 
 		float GetLastRunTimeMSec() const;
 		UE::NNEProfiling::Internal::FStatistics GetRunStatistics() const;
@@ -212,8 +212,8 @@ namespace NNX
 		TArray<const char*> InputTensorNames;
 		TArray<const char*> OutputTensorNames;
 
-		TArray<FMLTensorDesc> InputTensors;
-		TArray<FMLTensorDesc> OutputTensors;
+		TArray<FTensor> InputTensors;
+		TArray<FTensor> OutputTensors;
 
 		virtual bool InitializedAndConfigureMembers();
 		bool ConfigureTensors(const bool InIsInput);
diff --git a/Engine/Plugins/Experimental/NNX/Source/NNXRuntimeORT/Private/NNXRuntimeORTUtils.h b/Engine/Plugins/Experimental/NNX/Source/NNXRuntimeORT/Private/NNXRuntimeORTUtils.h
--- a/Engine/Plugins/Experimental/NNX/Source/NNXRuntimeORT/Private/NNXRuntimeORTUtils.h
+++ b/Engine/Plugins/Experimental/NNX/Source/NNXRuntimeORT/Private/NNXRuntimeORTUtils.h
@@ -113,7 +113,7 @@ inline TypeInfoORT TranslateTensorTypeORTToNNI(unsigned int OrtDataType)
 
 inline void BindTensorsToORT(
 	TArrayView<const NNX::FMLTensorBinding> InBindingTensors,
-	const TArray<NNX::FMLTensorDesc>& InTensorsDescriptors,
+	const TArray<NNX::FTensor>& InTensorsDescriptors,
 	const TArray<ONNXTensorElementDataType>& InTensorsORTType, 
 	const Ort::MemoryInfo* InAllocatorInfo,
 	TArray<Ort::Value>& OutOrtTensors	
@@ -131,18 +131,18 @@ inline void BindTensorsToORT(
 	for (uint32 Index = 0; Index < NumBinding; ++Index)
 	{
 		const NNX::FMLTensorBinding& CurrentBinding = InBindingTensors[Index];
-		const NNX::FMLTensorDesc& CurrentDescriptor = InTensorsDescriptors[Index];
+		const NNX::FTensor& CurrentDescriptor = InTensorsDescriptors[Index];
 		ONNXTensorElementDataType CurrentORTType = InTensorsORTType[Index];
 
 		TUniquePtr<int64_t[]> SizesInt64t;
-		SizesInt64t = MakeUnique<int64_t[]>(CurrentDescriptor.Shape.Num());
-		for (int32 DimIndex = 0; DimIndex < CurrentDescriptor.Shape.Num(); ++DimIndex)
+		SizesInt64t = MakeUnique<int64_t[]>(CurrentDescriptor.GetShape().Rank());
+		for (int32 DimIndex = 0; DimIndex < CurrentDescriptor.GetShape().Rank(); ++DimIndex)
 		{
-			SizesInt64t.Get()[DimIndex] = CurrentDescriptor.Shape.Data[DimIndex];
+			SizesInt64t.Get()[DimIndex] = CurrentDescriptor.GetShape().Data[DimIndex];
 		}
 
-		const uint64 ByteCount { InTensorsDescriptors[Index].DataSize };
-		const uint32 ArrayDimensions { (uint32)CurrentDescriptor.Shape.Num() };
+		const uint64 ByteCount { InTensorsDescriptors[Index].GetDataSize() };
+		const uint32 ArrayDimensions { (uint32)CurrentDescriptor.GetShape().Rank() };
 		OutOrtTensors.Emplace(
 			Ort::Value::CreateTensor(
 				*InAllocatorInfo,
diff --git a/Engine/Plugins/Experimental/NNX/Source/NNXRuntimeRDG/Private/Hlsl/NNIRuntimeRDGConv.cpp b/Engine/Plugins/Experimental/NNX/Source/NNXRuntimeRDG/Private/Hlsl/NNIRuntimeRDGConv.cpp
--- a/Engine/Plugins/Experimental/NNX/Source/NNXRuntimeRDG/Private/Hlsl/NNIRuntimeRDGConv.cpp
+++ b/Engine/Plugins/Experimental/NNX/Source/NNXRuntimeRDG/Private/Hlsl/NNIRuntimeRDGConv.cpp
@@ -29,10 +29,10 @@ namespace UE::NNIRuntimeRDG::Private::Hlsl
 
 		FConv() {}
 
-		NNX::FMLTensorDesc Input = {};
-		NNX::FMLTensorDesc Weights = {};
-		NNX::FMLTensorDesc Bias = {};
-		NNX::FMLTensorDesc Output = {};
+		NNX::FTensor Input = {};
+		NNX::FTensor Weights = {};
+		NNX::FTensor Bias = {};
+		NNX::FTensor Output = {};
 
 		int NumDimensions = 0;
 		bool HasBias = false;
@@ -46,7 +46,7 @@ namespace UE::NNIRuntimeRDG::Private::Hlsl
 
 	public:
 
-		virtual bool Initialize(TArrayView<const NNX::FMLTensorDesc> InputTensors, TArrayView<const NNX::FMLTensorDesc> OutputTensors, const UE::NNECore::FAttributeMap& Attributes) override
+		virtual bool Initialize(TArrayView<const NNX::FTensor> InputTensors, TArrayView<const NNX::FTensor> OutputTensors, const UE::NNECore::FAttributeMap& Attributes) override
 		{
 			check(InputTensors.Num() >= 2 && InputTensors.Num() <= 3);
 			check(OutputTensors.Num() == 1);
@@ -55,9 +55,9 @@ namespace UE::NNIRuntimeRDG::Private::Hlsl
 			Weights = InputTensors[1];
 			Output = OutputTensors[0];
 
-			check(Input.Shape.Num() > 2);
-			check(Weights.Shape.Num()  == Input.Shape.Num());
-			check(Output.Shape.Num() == Input.Shape.Num());
+			check(Input.GetShape().Rank() > 2);
+			check(Weights.GetShape().Rank()  == Input.GetShape().Rank());
+			check(Output.GetShape().Rank() == Input.GetShape().Rank());
 
 			if (InputTensors.Num() == 3) {
 				HasBias = true;
@@ -67,7 +67,7 @@ namespace UE::NNIRuntimeRDG::Private::Hlsl
 				HasBias = false;
 			}
 
-			NumDimensions = Input.Shape.Num() - 2;
+			NumDimensions = Input.GetShape().Rank() - 2;
 
 			return true;
 		}
@@ -79,11 +79,11 @@ namespace UE::NNIRuntimeRDG::Private::Hlsl
 			constexpr EConvAlgorithm Algorithm = EConvAlgorithm::SharedMemory;
 			constexpr EConvGroupSize GroupSize = EConvGroupSize::Size256;
 
-			TArray<int32> OutputShape = FConvCS::GetOutputShape(Input.Shape.Data, Weights.Shape.Data, AutoPad, Dilations, Strides, Pads);
+			TArray<int32> OutputShape = FConvCS::GetOutputShape(Input.GetShape().Data, Weights.GetShape().Data, AutoPad, Dilations, Strides, Pads);
 
 			// Set parameters
 			FConvCS::FParameters* Params = GraphBuilder.AllocParameters<FConvCS::FParameters>();
-			FConvCS::FillInParameters(GroupSize, Input.Shape.Data, Weights.Shape.Data, HasBias, AutoPad, Group, Dilations,Strides, Pads, *Params);
+			FConvCS::FillInParameters(GroupSize, Input.GetShape().Data, Weights.GetShape().Data, HasBias, AutoPad, Group, Dilations,Strides, Pads, *Params);
 			Params->X = GraphBuilder.CreateSRV(FRDGBufferSRVDesc(InInputBindings[0].Buffer, PF_R32_FLOAT));
 			Params->W = GraphBuilder.CreateSRV(FRDGBufferSRVDesc(InInputBindings[1].Buffer, PF_R32_FLOAT));
 			if (InInputBindings.Num() == 3) {
@@ -96,7 +96,7 @@ namespace UE::NNIRuntimeRDG::Private::Hlsl
 			PermutationVector.Set<FConvCS::FConvAlgorithm>(Algorithm);
 			PermutationVector.Set<FConvCS::FConvGroupSize>(GroupSize);
 			PermutationVector.Set<FConvCS::FConvNumDimensions>(NumDimensions);
-			PermutationVector.Set<FConvCS::FConvNumReadsPerThread>(FConvCS::GetNumReadsPerThread(GroupSize, Weights.Shape.Data, Dilations, Strides));
+			PermutationVector.Set<FConvCS::FConvNumReadsPerThread>(FConvCS::GetNumReadsPerThread(GroupSize, Weights.GetShape().Data, Dilations, Strides));
 			PermutationVector.Set<FConvCS::FConvHasB>(HasBias);
 			TShaderMapRef<FConvCS> ComputeShader(GetGlobalShaderMap(GMaxRHIFeatureLevel), PermutationVector);
 
diff --git a/Engine/Plugins/Experimental/NNX/Source/NNXRuntimeRDG/Private/Hlsl/NNIRuntimeRDGConvTranspose.cpp b/Engine/Plugins/Experimental/NNX/Source/NNXRuntimeRDG/Private/Hlsl/NNIRuntimeRDGConvTranspose.cpp
--- a/Engine/Plugins/Experimental/NNX/Source/NNXRuntimeRDG/Private/Hlsl/NNIRuntimeRDGConvTranspose.cpp
+++ b/Engine/Plugins/Experimental/NNX/Source/NNXRuntimeRDG/Private/Hlsl/NNIRuntimeRDGConvTranspose.cpp
@@ -28,10 +28,10 @@ namespace UE::NNIRuntimeRDG::Private::Hlsl
 
 		FConvTranspose() {}
 
-		NNX::FMLTensorDesc Input = {};
-		NNX::FMLTensorDesc Weights = {};
-		NNX::FMLTensorDesc Bias = {};
-		NNX::FMLTensorDesc Output = {};
+		NNX::FTensor Input = {};
+		NNX::FTensor Weights = {};
+		NNX::FTensor Bias = {};
+		NNX::FTensor Output = {};
 
 		int NumDimensions = 0;
 		bool HasBias = false;
@@ -46,7 +46,7 @@ namespace UE::NNIRuntimeRDG::Private::Hlsl
 
 	public:
 
-		virtual bool Initialize(TArrayView<const NNX::FMLTensorDesc> InputTensors, TArrayView<const NNX::FMLTensorDesc> OutputTensors, const UE::NNECore::FAttributeMap& Attributes) override
+		virtual bool Initialize(TArrayView<const NNX::FTensor> InputTensors, TArrayView<const NNX::FTensor> OutputTensors, const UE::NNECore::FAttributeMap& Attributes) override
 		{
 			check(InputTensors.Num() >= 2 && InputTensors.Num() <= 3);
 			check(OutputTensors.Num() == 1);
@@ -55,9 +55,9 @@ namespace UE::NNIRuntimeRDG::Private::Hlsl
 			Weights = InputTensors[1];
 			Output = OutputTensors[0];
 
-			check(Input.Shape.Num() > 2);
-			check(Weights.Shape.Num()  == Input.Shape.Num());
-			check(Output.Shape.Num() == Input.Shape.Num());
+			check(Input.GetShape().Rank() > 2);
+			check(Weights.GetShape().Rank()  == Input.GetShape().Rank());
+			check(Output.GetShape().Rank() == Input.GetShape().Rank());
 
 			if (InputTensors.Num() == 3) {
 				HasBias = true;
@@ -68,7 +68,7 @@ namespace UE::NNIRuntimeRDG::Private::Hlsl
 				HasBias = false;
 			}
 
-			NumDimensions = Input.Shape.Num() - 2;
+			NumDimensions = Input.GetShape().Rank() - 2;
 
 			return true;
 		}
@@ -80,11 +80,11 @@ namespace UE::NNIRuntimeRDG::Private::Hlsl
 			constexpr EConvTransposeAlgorithm Algorithm = EConvTransposeAlgorithm::SharedMemory;
 			constexpr EConvTransposeGroupSize GroupSize = EConvTransposeGroupSize::Size256;
 
-			TArray<int32> OutputShape = FConvTransposeCS::GetOutputShape(Input.Shape.Data, Weights.Shape.Data, AutoPad, Dilations, Strides, Pads, OutputPadding, Group);
+			TArray<int32> OutputShape = FConvTransposeCS::GetOutputShape(Input.GetShape().Data, Weights.GetShape().Data, AutoPad, Dilations, Strides, Pads, OutputPadding, Group);
 
 			// Set parameters
 			FConvTransposeCS::FParameters* Params = GraphBuilder.AllocParameters<FConvTransposeCS::FParameters>();
-			FConvTransposeCS::FillInParameters(GroupSize, Input.Shape.Data, Weights.Shape.Data, HasBias, AutoPad, Group, Dilations,Strides, Pads, OutputPadding, *Params);
+			FConvTransposeCS::FillInParameters(GroupSize, Input.GetShape().Data, Weights.GetShape().Data, HasBias, AutoPad, Group, Dilations,Strides, Pads, OutputPadding, *Params);
 			Params->X = GraphBuilder.CreateSRV(FRDGBufferSRVDesc(InInputBindings[0].Buffer, PF_R32_FLOAT));
 			Params->W = GraphBuilder.CreateSRV(FRDGBufferSRVDesc(InInputBindings[1].Buffer, PF_R32_FLOAT));
 			if (InInputBindings.Num() == 3) {
@@ -97,7 +97,7 @@ namespace UE::NNIRuntimeRDG::Private::Hlsl
 			PermutationVector.Set<FConvTransposeCS::FConvTransposeAlgorithm>(Algorithm);
 			PermutationVector.Set<FConvTransposeCS::FConvTransposeGroupSize>(GroupSize);
 			PermutationVector.Set<FConvTransposeCS::FConvTransposeNumStackDimensions>(NumDimensions);
-			PermutationVector.Set<FConvTransposeCS::FConvTransposeNumReadsPerThread>(FConvTransposeCS::GetNumReadsPerThread(GroupSize, Weights.Shape.Data, Dilations, Strides));
+			PermutationVector.Set<FConvTransposeCS::FConvTransposeNumReadsPerThread>(FConvTransposeCS::GetNumReadsPerThread(GroupSize, Weights.GetShape().Data, Dilations, Strides));
 			PermutationVector.Set<FConvTransposeCS::FConvTransposeHasB>(HasBias);
 			TShaderMapRef<FConvTransposeCS> ComputeShader(GetGlobalShaderMap(GMaxRHIFeatureLevel), PermutationVector);
 
diff --git a/Engine/Plugins/Experimental/NNX/Source/NNXRuntimeRDG/Private/Hlsl/NNIRuntimeRDGElementWiseBinary.cpp b/Engine/Plugins/Experimental/NNX/Source/NNXRuntimeRDG/Private/Hlsl/NNIRuntimeRDGElementWiseBinary.cpp
--- a/Engine/Plugins/Experimental/NNX/Source/NNXRuntimeRDG/Private/Hlsl/NNIRuntimeRDGElementWiseBinary.cpp
+++ b/Engine/Plugins/Experimental/NNX/Source/NNXRuntimeRDG/Private/Hlsl/NNIRuntimeRDGElementWiseBinary.cpp
@@ -24,13 +24,13 @@ namespace UE::NNIRuntimeRDG::Private::Hlsl
 
 	private:
 
-		NNX::FMLTensorDesc LHSInput = {};
-		NNX::FMLTensorDesc RHSInput = {};
-		NNX::FMLTensorDesc Output = {};
+		NNX::FTensor LHSInput = {};
+		NNX::FTensor RHSInput = {};
+		NNX::FTensor Output = {};
 
 	public:
 
-		virtual bool Initialize(TArrayView<const NNX::FMLTensorDesc> InputTensors, TArrayView<const NNX::FMLTensorDesc> OutputTensors, const UE::NNECore::FAttributeMap& Attributes) override
+		virtual bool Initialize(TArrayView<const NNX::FTensor> InputTensors, TArrayView<const NNX::FTensor> OutputTensors, const UE::NNECore::FAttributeMap& Attributes) override
 		{
 			check(InputTensors.Num() == 2);
 			check(OutputTensors.Num() == 1);
@@ -49,23 +49,23 @@ namespace UE::NNIRuntimeRDG::Private::Hlsl
 			FRDGBufferSRVRef RHSInputSRV = GraphBuilder.CreateSRV(FRDGBufferSRVDesc(InInputBindings[1].Buffer, PF_R32_FLOAT));
 			FRDGBufferUAVRef OutputUAV = GraphBuilder.CreateUAV(FRDGBufferUAVDesc(OutOutputBindings[0].Buffer, PF_R32_FLOAT));
 
-			FIntVector ThreadGroupCount = NNX::ComputeElementWiseThreadGroups(Output.Volume, FElementWiseBinaryConstants::NUM_GROUP_THREADS);
+			FIntVector ThreadGroupCount = NNX::ComputeElementWiseThreadGroups(Output.GetVolume(), FElementWiseBinaryConstants::NUM_GROUP_THREADS);
 
 			// Set parameters
 			TElementWiseBinaryCS::FParameters* Params = GraphBuilder.AllocParameters<TElementWiseBinaryCS::FParameters>();
 			Params->LHSInput = LHSInputSRV;
 			Params->RHSInput = RHSInputSRV;
 			Params->Output = OutputUAV;
-			FillTensorStrideForBroadcastShaderParameters(LHSInput, Output.Shape.Num(), Params->TensorInfo, 0);
-			FillTensorStrideForBroadcastShaderParameters(RHSInput, Output.Shape.Num(), Params->TensorInfo, 1);
+			FillTensorStrideForBroadcastShaderParameters(LHSInput, Output.GetShape().Rank(), Params->TensorInfo, 0);
+			FillTensorStrideForBroadcastShaderParameters(RHSInput, Output.GetShape().Rank(), Params->TensorInfo, 1);
 			FillTensorStrideShaderParameters(Output, Params->TensorInfo, 2);
-			Params->Num = Output.Volume;
+			Params->Num = Output.GetVolume();
 			Params->ThreadCountX = ThreadGroupCount.X * FElementWiseBinaryConstants::NUM_GROUP_THREADS;
 
 			TElementWiseBinaryCS::FPermutationDomain PermutationVector;
 
 			PermutationVector.Set<TElementWiseBinaryCS::FOperatorType>(OpType);
-			PermutationVector.Set<TElementWiseBinaryCS::FBinaryNumDimensions>(Output.Shape.Num());
+			PermutationVector.Set<TElementWiseBinaryCS::FBinaryNumDimensions>(Output.GetShape().Rank());
 
 			TShaderMapRef<TElementWiseBinaryCS> ComputeShader(GetGlobalShaderMap(GMaxRHIFeatureLevel), PermutationVector);
 		
diff --git a/Engine/Plugins/Experimental/NNX/Source/NNXRuntimeRDG/Private/Hlsl/NNIRuntimeRDGElementWiseUnary.cpp b/Engine/Plugins/Experimental/NNX/Source/NNXRuntimeRDG/Private/Hlsl/NNIRuntimeRDGElementWiseUnary.cpp
--- a/Engine/Plugins/Experimental/NNX/Source/NNXRuntimeRDG/Private/Hlsl/NNIRuntimeRDGElementWiseUnary.cpp
+++ b/Engine/Plugins/Experimental/NNX/Source/NNXRuntimeRDG/Private/Hlsl/NNIRuntimeRDGElementWiseUnary.cpp
@@ -28,12 +28,12 @@ namespace UE::NNIRuntimeRDG::Private::Hlsl
 		float Alpha = 0.0f;
 		float Beta = 0.0f;
 		float Gamma = 0.0f;
-		NNX::FMLTensorDesc Input = {};
-		NNX::FMLTensorDesc Output = {};
+		NNX::FTensor Input = {};
+		NNX::FTensor Output = {};
 
 	public:
 
-		virtual bool Initialize(TArrayView<const NNX::FMLTensorDesc> InputTensors, TArrayView<const NNX::FMLTensorDesc> OutputTensors, const UE::NNECore::FAttributeMap& Attributes) override
+		virtual bool Initialize(TArrayView<const NNX::FTensor> InputTensors, TArrayView<const NNX::FTensor> OutputTensors, const UE::NNECore::FAttributeMap& Attributes) override
 		{
 			check(InputTensors.Num() == 1);
 			check(OutputTensors.Num() == 1);
@@ -57,28 +57,28 @@ namespace UE::NNIRuntimeRDG::Private::Hlsl
             //}
 
 			//Dispatch(FRDGBuilder & GraphBuilder, 
-			         //TConstArrayView<const NNX::RDGTensor> InputBuffers,
-					 //TConstArrayView<const NNX::RDGTensor> OutputBuffers) override
+			         //TConstArrayView<NNX::RDGTensor> InputBuffers,
+					 //TConstArrayView<NNX::RDGTensor> OutputBuffers) override
 
             //or 
             //Dispatch(FRDGBuilder & GraphBuilder, 
-                     //TConstArrayView<const NNX::FConcreteShape> InputShapes,
-			         //TConstArrayView<const NNX::RDGTensor*> InputBuffers,
-                     //TConstArrayView<const NNX::FConcreteShape> OutputShapes,
-					 //TConstArrayView<const NNX::RDGTensor*> OutputBuffers) override
+                     //TConstArrayView<NNX::FConcreteShape> InputShapes,
+			         //TConstArrayView<NNX::RDGTensor*> InputBuffers,
+                     //TConstArrayView<NNX::FConcreteShape> OutputShapes,
+					 //TConstArrayView<NNX::RDGTensor*> OutputBuffers) override
 
             //or 
             //Dispatch(FRDGBuilder & GraphBuilder, 
-                     //TConstArrayView<const NNX::RDGTensor*> Buffers,
-                     //TConstArrayView<const NNX::FConcreteShape> Shapes,
-			         //TConstArrayView<const uint32> InputIndices,
-                     //TConstArrayView<const uint32> OutputIndices) override
+                     //TConstArrayView<NNX::RDGTensor*> Buffers,
+                     //TConstArrayView<NNX::FConcreteShape> Shapes,
+			         //TConstArrayView<uint32> InputIndices,
+                     //TConstArrayView<uint32> OutputIndices) override
 
 			// HACK: This only works for single layer networks
 			FRDGBufferSRVRef InputSRV = GraphBuilder.CreateSRV(FRDGBufferSRVDesc(InInputBindings[0].Buffer, PF_R32_FLOAT));
 			FRDGBufferUAVRef OutputUAV = GraphBuilder.CreateUAV(FRDGBufferUAVDesc(OutOutputBindings[0].Buffer, PF_R32_FLOAT));
 		
-			FIntVector ThreadGroupCount = NNX::ComputeElementWiseThreadGroups(Output.Volume, FElementWiseUnaryConstants::NUM_GROUP_THREADS);
+			FIntVector ThreadGroupCount = NNX::ComputeElementWiseThreadGroups(Output.GetVolume(), FElementWiseUnaryConstants::NUM_GROUP_THREADS);
 
 			// Set parameters
 			TElementWiseUnaryCS::FParameters* Params = GraphBuilder.AllocParameters<TElementWiseUnaryCS::FParameters>();
@@ -87,7 +87,7 @@ namespace UE::NNIRuntimeRDG::Private::Hlsl
 			Params->Alpha = Alpha;
 			Params->Beta = Beta;
 			Params->Gamma = Gamma;
-			Params->Num = Output.Volume;
+			Params->Num = Output.GetVolume();
 			Params->ThreadCountX = ThreadGroupCount.X * FElementWiseUnaryConstants::NUM_GROUP_THREADS;
 
 			TElementWiseUnaryCS::FPermutationDomain PermutationVector;
diff --git a/Engine/Plugins/Experimental/NNX/Source/NNXRuntimeRDG/Private/Hlsl/NNIRuntimeRDGElementWiseVariadic.cpp b/Engine/Plugins/Experimental/NNX/Source/NNXRuntimeRDG/Private/Hlsl/NNIRuntimeRDGElementWiseVariadic.cpp
--- a/Engine/Plugins/Experimental/NNX/Source/NNXRuntimeRDG/Private/Hlsl/NNIRuntimeRDGElementWiseVariadic.cpp
+++ b/Engine/Plugins/Experimental/NNX/Source/NNXRuntimeRDG/Private/Hlsl/NNIRuntimeRDGElementWiseVariadic.cpp
@@ -13,9 +13,9 @@ namespace UE::NNIRuntimeRDG::Private::Hlsl
 
 	void AddOneVariadicOpPass(FRDGBuilder& GraphBuilder, 
 		TArrayView<const NNX::FMLTensorBinding> InputBindings,
-		TArrayView<const NNX::FMLTensorDesc> InputDesc,
+		TArrayView<const NNX::FTensor> InputDesc,
 		const NNX::FMLTensorBinding& OutputBinding,
-		const NNX::FMLTensorDesc& OutputDesc,
+		const NNX::FTensor& OutputDesc,
 		bool OutputAsInput,
 		EMLElementWiseVariadicOperatorType OpType,
 		float Scale)
@@ -35,7 +35,7 @@ namespace UE::NNIRuntimeRDG::Private::Hlsl
 		}
 
 		// Set parameters
-		FIntVector ThreadGroupCount = NNX::ComputeElementWiseThreadGroups(OutputDesc.Volume, FElementWiseVariadicConstants::NUM_GROUP_THREADS);
+		FIntVector ThreadGroupCount = NNX::ComputeElementWiseThreadGroups(OutputDesc.GetVolume(), FElementWiseVariadicConstants::NUM_GROUP_THREADS);
 		TElementWiseVariadicCS::FParameters* Params = GraphBuilder.AllocParameters<TElementWiseVariadicCS::FParameters>();
 
 		Params->Input0 = InputsSRV[0];
@@ -43,21 +43,21 @@ namespace UE::NNIRuntimeRDG::Private::Hlsl
 		Params->Input2 = InputsSRV[2];
 		Params->Input3 = InputsSRV[3];
 		Params->Output = OutputUAV;
-		FillTensorStrideForBroadcastShaderParameters(InputDesc[0], OutputDesc.Shape.Num(), Params->InputTensorInfo, 0);
+		FillTensorStrideForBroadcastShaderParameters(InputDesc[0], OutputDesc.GetShape().Rank(), Params->InputTensorInfo, 0);
 		if (InputBindings.Num() >= 2)
 		{
-			FillTensorStrideForBroadcastShaderParameters(InputDesc[1], OutputDesc.Shape.Num(), Params->InputTensorInfo, 1);
+			FillTensorStrideForBroadcastShaderParameters(InputDesc[1], OutputDesc.GetShape().Rank(), Params->InputTensorInfo, 1);
 		}
 		if (InputBindings.Num() >= 3)
 		{
-			FillTensorStrideForBroadcastShaderParameters(InputDesc[2], OutputDesc.Shape.Num(), Params->InputTensorInfo, 2);
+			FillTensorStrideForBroadcastShaderParameters(InputDesc[2], OutputDesc.GetShape().Rank(), Params->InputTensorInfo, 2);
 		}
 		if (InputBindings.Num() >= 4)
 		{
-			FillTensorStrideForBroadcastShaderParameters(InputDesc[3], OutputDesc.Shape.Num(), Params->InputTensorInfo, 3);
+			FillTensorStrideForBroadcastShaderParameters(InputDesc[3], OutputDesc.GetShape().Rank(), Params->InputTensorInfo, 3);
 		}
 		FillTensorStrideShaderParameters(OutputDesc, Params->OutputTensorInfo, 0);
-		Params->Num = OutputDesc.Volume;
+		Params->Num = OutputDesc.GetVolume();
 		Params->ThreadCountX = ThreadGroupCount.X * FElementWiseVariadicConstants::NUM_GROUP_THREADS;
 		Params->Scale = Scale;
 
@@ -68,7 +68,7 @@ namespace UE::NNIRuntimeRDG::Private::Hlsl
 		PermutationVector.Set<TElementWiseVariadicCS::FApplyScale>(Scale != 1.0f);
 		PermutationVector.Set<TElementWiseVariadicCS::FOutputAsInput>(OutputAsInput);
 		PermutationVector.Set<TElementWiseVariadicCS::FNumInput>(InputBindings.Num());
-		PermutationVector.Set<TElementWiseVariadicCS::FVariadicNumDimensions>(OutputDesc.Shape.Num());
+		PermutationVector.Set<TElementWiseVariadicCS::FVariadicNumDimensions>(OutputDesc.GetShape().Rank());
 
 		// Add the pass to RDG
 		TShaderMapRef<TElementWiseVariadicCS> ComputeShader(GetGlobalShaderMap(GMaxRHIFeatureLevel), PermutationVector);
@@ -94,12 +94,12 @@ namespace UE::NNIRuntimeRDG::Private::Hlsl
 
 	private:
 
-		TArray<NNX::FMLTensorDesc> InputDescs;
-		NNX::FMLTensorDesc OutputDesc = {};
+		TArray<NNX::FTensor> InputDescs;
+		NNX::FTensor OutputDesc = {};
 
 	public:
 
-		virtual bool Initialize(TArrayView<const NNX::FMLTensorDesc> InputTensors, TArrayView<const NNX::FMLTensorDesc> OutputTensors, const UE::NNECore::FAttributeMap& Attributes) override
+		virtual bool Initialize(TArrayView<const NNX::FTensor> InputTensors, TArrayView<const NNX::FTensor> OutputTensors, const UE::NNECore::FAttributeMap& Attributes) override
 		{
 			check(InputTensors.Num() > 0);
 			check(OutputTensors.Num() == 1);
@@ -119,7 +119,7 @@ namespace UE::NNIRuntimeRDG::Private::Hlsl
 			RDG_GPU_STAT_SCOPE(GraphBuilder, FNNIOperatorElementWiseVariadic);
 
 			NNX::FMLTensorBinding PassInputBindings[FElementWiseVariadicConstants::MAX_NUM_INPUT];
-			NNX::FMLTensorDesc PassInputDescs[FElementWiseVariadicConstants::MAX_NUM_INPUT];
+			NNX::FTensor PassInputDescs[FElementWiseVariadicConstants::MAX_NUM_INPUT];
 			for (int32 InputOffset = 0; InputOffset < InInputBindings.Num(); InputOffset += FElementWiseVariadicConstants::MAX_NUM_INPUT)
 			{
 				uint32 NumInputLeftToHandle = InInputBindings.Num() - InputOffset;
diff --git a/Engine/Plugins/Experimental/NNX/Source/NNXRuntimeRDG/Private/Hlsl/NNIRuntimeRDGGather.cpp b/Engine/Plugins/Experimental/NNX/Source/NNXRuntimeRDG/Private/Hlsl/NNIRuntimeRDGGather.cpp
--- a/Engine/Plugins/Experimental/NNX/Source/NNXRuntimeRDG/Private/Hlsl/NNIRuntimeRDGGather.cpp
+++ b/Engine/Plugins/Experimental/NNX/Source/NNXRuntimeRDG/Private/Hlsl/NNIRuntimeRDGGather.cpp
@@ -23,27 +23,27 @@ namespace UE::NNIRuntimeRDG::Private::Hlsl
 	private:
 
 		int32 Axis = 0;
-		NNX::FMLTensorDesc Data = {};
-		NNX::FMLTensorDesc Indices = {};
-		NNX::FMLTensorDesc Output = {};
+		NNX::FTensor Data = {};
+		NNX::FTensor Indices = {};
+		NNX::FTensor Output = {};
 
 	public:
 
-		virtual bool Initialize(TArrayView<const NNX::FMLTensorDesc> InputTensors, TArrayView<const NNX::FMLTensorDesc> OutputTensors, const UE::NNECore::FAttributeMap& Attributes) override
+		virtual bool Initialize(TArrayView<const NNX::FTensor> InputTensors, TArrayView<const NNX::FTensor> OutputTensors, const UE::NNECore::FAttributeMap& Attributes) override
 		{
 			using namespace UE::NNEHlslShaders::Internal;
 
 			check(InputTensors.Num() == 2)
 			check(OutputTensors.Num() == 1)
-			check(Output.Shape.Num() <= FGatherConstants::MAX_NUM_DIMENSIONS)
-			check(InputTensors[0].Shape.Num() > 0)
-			check(InputTensors[1].Shape.Num() > 0)
-			check(InputTensors[0].Shape.Num() + (InputTensors[1].Shape.Num() - 1) <= FGatherConstants::MAX_NUM_DIMENSIONS)
+			check(Output.GetShape().Rank() <= FGatherConstants::MAX_NUM_DIMENSIONS)
+			check(InputTensors[0].GetShape().Rank() > 0)
+			check(InputTensors[1].GetShape().Rank() > 0)
+			check(InputTensors[0].GetShape().Rank() + (InputTensors[1].GetShape().Rank() - 1) <= FGatherConstants::MAX_NUM_DIMENSIONS)
 
 			Axis = Attributes.GetValueOrDefault(TEXT("axis"), Axis);
-			check(Axis < InputTensors[0].Shape.Num())
-			check(Axis >= -InputTensors[0].Shape.Num())
-			Axis = Axis >= 0 ? Axis : InputTensors[0].Shape.Num() + Axis;
+			check(Axis < InputTensors[0].GetShape().Rank())
+			check(Axis >= -InputTensors[0].GetShape().Rank())
+			Axis = Axis >= 0 ? Axis : InputTensors[0].GetShape().Rank() + Axis;
 
 			Data = InputTensors[0];
 			Indices = InputTensors[1];
@@ -64,7 +64,7 @@ namespace UE::NNIRuntimeRDG::Private::Hlsl
 			Parameters->Output = GraphBuilder.CreateUAV(FRDGBufferUAVDesc(OutOutputBindings[0].Buffer, PF_R32_FLOAT));
 
 			TGatherCS::FPermutationDomain PermutationVector;
-			PermutationVector.Set<TGatherCS::FGatherNumOutputDimensions>(Output.Shape.Num());
+			PermutationVector.Set<TGatherCS::FGatherNumOutputDimensions>(Output.GetShape().Rank());
 			TShaderMapRef<TGatherCS> ComputeShader(GetGlobalShaderMap(GMaxRHIFeatureLevel), PermutationVector);
 
 			FIntVector ThreadGroupCount = TGatherCS::GetGroupCount(*Parameters);
diff --git a/Engine/Plugins/Experimental/NNX/Source/NNXRuntimeRDG/Private/Hlsl/NNIRuntimeRDGGemm.cpp b/Engine/Plugins/Experimental/NNX/Source/NNXRuntimeRDG/Private/Hlsl/NNIRuntimeRDGGemm.cpp
--- a/Engine/Plugins/Experimental/NNX/Source/NNXRuntimeRDG/Private/Hlsl/NNIRuntimeRDGGemm.cpp
+++ b/Engine/Plugins/Experimental/NNX/Source/NNXRuntimeRDG/Private/Hlsl/NNIRuntimeRDGGemm.cpp
@@ -21,10 +21,10 @@ namespace UE::NNIRuntimeRDG::Private::Hlsl
 
 	private:
 
-		NNX::FMLTensorDesc InputA = {};
-		NNX::FMLTensorDesc InputB = {};
-		NNX::FMLTensorDesc InputC = {};
-		NNX::FMLTensorDesc Output = {};
+		NNX::FTensor InputA = {};
+		NNX::FTensor InputB = {};
+		NNX::FTensor InputC = {};
+		NNX::FTensor Output = {};
 
 		float InputAlpha = 1.0f;
 		float InputBeta = 1.0f;
@@ -41,7 +41,7 @@ namespace UE::NNIRuntimeRDG::Private::Hlsl
 
 	public:
 
-		virtual bool Initialize(TArrayView<const NNX::FMLTensorDesc> InputTensors, TArrayView<const NNX::FMLTensorDesc> OutputTensors, const UE::NNECore::FAttributeMap& Attributes) override
+		virtual bool Initialize(TArrayView<const NNX::FTensor> InputTensors, TArrayView<const NNX::FTensor> OutputTensors, const UE::NNECore::FAttributeMap& Attributes) override
 		{
 			check(InputTensors.Num() >= 2 && InputTensors.Num() <= 3);
 			check(OutputTensors.Num() == 1);
@@ -51,11 +51,11 @@ namespace UE::NNIRuntimeRDG::Private::Hlsl
 			if (InputTensors.Num() == 3) InputC = InputTensors[2];
 			Output = OutputTensors[0];
 
-			check(InputA.Shape.Num() == 2);
-			check(InputB.Shape.Num() == 2);
-			check(InputC.Shape.Num() < 3);
+			check(InputA.GetShape().Rank() == 2);
+			check(InputB.GetShape().Rank() == 2);
+			check(InputC.GetShape().Rank() < 3);
 
-			check(InputC.Shape.Num() != 1 || InputC.Shape.Data[0] != 1); // TODO scalar version not supported yet
+			check(InputC.GetShape().Rank() != 1 || InputC.GetShape().Data[0] != 1); // TODO scalar version not supported yet
 
 			// C is treated as a scalar if there is no valid C, either width or height is zero or C dimension is 1x1
 			bIsCScalar = false; // InputTensors.Num() != 3 || InputC.Sizes[0] * InputC.Sizes[1] < 2;
@@ -109,7 +109,7 @@ namespace UE::NNIRuntimeRDG::Private::Hlsl
 		}
 	};
 
-	bool ValidateGemmOperator(const UE::NNECore::FAttributeMap& AttributeMap, TConstArrayView<EMLTensorDataType> InputTypes, TConstArrayView<const NNX::FSymbolicTensorShape> InputShapes)
+	bool ValidateGemmOperator(const UE::NNECore::FAttributeMap& AttributeMap, TConstArrayView<EMLTensorDataType> InputTypes, TConstArrayView<NNX::FSymbolicTensorShape> InputShapes)
 	{
 		bool bIsValid = true;
 
diff --git a/Engine/Plugins/Experimental/NNX/Source/NNXRuntimeRDG/Private/Hlsl/NNIRuntimeRDGMatMul.cpp b/Engine/Plugins/Experimental/NNX/Source/NNXRuntimeRDG/Private/Hlsl/NNIRuntimeRDGMatMul.cpp
--- a/Engine/Plugins/Experimental/NNX/Source/NNXRuntimeRDG/Private/Hlsl/NNIRuntimeRDGMatMul.cpp
+++ b/Engine/Plugins/Experimental/NNX/Source/NNXRuntimeRDG/Private/Hlsl/NNIRuntimeRDGMatMul.cpp
@@ -20,13 +20,13 @@ namespace UE::NNIRuntimeRDG::Private::Hlsl
 
 	private:
 
-		NNX::FMLTensorDesc InputA = {};
-		NNX::FMLTensorDesc InputB = {};
-		NNX::FMLTensorDesc Output = {};
+		NNX::FTensor InputA = {};
+		NNX::FTensor InputB = {};
+		NNX::FTensor Output = {};
 
 	public:
 
-		virtual bool Initialize(TArrayView<const NNX::FMLTensorDesc> InputTensors, TArrayView<const NNX::FMLTensorDesc> OutputTensors, const UE::NNECore::FAttributeMap& Attributes) override
+		virtual bool Initialize(TArrayView<const NNX::FTensor> InputTensors, TArrayView<const NNX::FTensor> OutputTensors, const UE::NNECore::FAttributeMap& Attributes) override
 		{
 			check(InputTensors.Num() == 2);
 			check(OutputTensors.Num() == 1);
@@ -44,7 +44,7 @@ namespace UE::NNIRuntimeRDG::Private::Hlsl
 
 			const EGemmAlgorithm Algorithm = EGemmAlgorithm::Simple32x32;
 
-			const int32 NumStackDimensions = FMath::Max(FMath::Max(InputA.Shape.Num(), InputB.Shape.Num()) - 2, 0);
+			const int32 NumStackDimensions = FMath::Max(FMath::Max(InputA.GetShape().Rank(), InputB.GetShape().Rank()) - 2, 0);
 
 			// Set parameters
 			TGemmCS::FParameters* Parameters = GraphBuilder.AllocParameters<TGemmCS::FParameters>();
diff --git a/Engine/Plugins/Experimental/NNX/Source/NNXRuntimeRDG/Private/NNXRuntimeDML.cpp b/Engine/Plugins/Experimental/NNX/Source/NNXRuntimeRDG/Private/NNXRuntimeDML.cpp
--- a/Engine/Plugins/Experimental/NNX/Source/NNXRuntimeRDG/Private/NNXRuntimeDML.cpp
+++ b/Engine/Plugins/Experimental/NNX/Source/NNXRuntimeRDG/Private/NNXRuntimeDML.cpp
@@ -167,59 +167,59 @@ namespace DmlUtil
 		TArray<uint32, TInlineAllocator<FTensorShape::MaxRank>>		Strides;
 	};
 
-	void SetTensorStrides(FTensorDesc& TensorDesc, const FMLTensorDesc& InputDesc)
+	void SetTensorStrides(FTensorDesc& TensorDesc, const FTensor& InputDesc)
 	{
 		uint32 CurrStride = 1;
 
-		TensorDesc.Strides.SetNum(InputDesc.Shape.Num());
+		TensorDesc.Strides.SetNum(InputDesc.GetShape().Rank());
 		
-		for (int32 i = InputDesc.Shape.Num() - 1; i >= 0; --i)
+		for (int32 i = InputDesc.GetShape().Rank() - 1; i >= 0; --i)
 		{
 			TensorDesc.Strides[i] = CurrStride;
-			CurrStride *= InputDesc.Shape.Data[i];
+			CurrStride *= InputDesc.GetShape().Data[i];
 		}
 	}
 
-	void SetTensorSizesAndStridesForBroadcast(FTensorDesc& TensorDesc, const FMLTensorDesc& InputDesc, const FMLTensorDesc& TargetDesc)
+	void SetTensorSizesAndStridesForBroadcast(FTensorDesc& TensorDesc, const FTensor& InputDesc, const FTensor& TargetDesc)
 	{
 		static_assert(FTensorShape::MaxRank <= 8);
 		
-		const uint32 TargetDimension = TargetDesc.Shape.Num() != -1 ? TargetDesc.Shape.Num() : InputDesc.Shape.Num();
-		checkf(TargetDesc.Shape.Num() >= InputDesc.Shape.Num(), TEXT("Can't broadcast tensor from rank %d to rank %d, should be inferior or equal."), InputDesc.Shape.Num(), TargetDimension);
+		const uint32 TargetDimension = TargetDesc.GetShape().Rank() != -1 ? TargetDesc.GetShape().Rank() : InputDesc.GetShape().Rank();
+		checkf(TargetDesc.GetShape().Rank() >= InputDesc.GetShape().Rank(), TEXT("Can't broadcast tensor from rank %d to rank %d, should be inferior or equal."), InputDesc.GetShape().Rank(), TargetDimension);
 		
 		TensorDesc.Sizes.SetNum(TargetDimension);
 		TensorDesc.Strides.SetNum(TargetDimension);
 
-		const int32 DimensionOffset = int32(TargetDimension - InputDesc.Shape.Num());
+		const int32 DimensionOffset = int32(TargetDimension - InputDesc.GetShape().Rank());
 		
 		for (int32 i = 0; i < (int32) TargetDimension; ++i)
 		{
-			TensorDesc.Sizes[i] = i < DimensionOffset ? 1 : InputDesc.Shape.Data[i - DimensionOffset];
+			TensorDesc.Sizes[i] = i < DimensionOffset ? 1 : InputDesc.GetShape().Data[i - DimensionOffset];
 		}
 
 		uint32 CurrStride = 1;
 
 		for (int32 i = TargetDimension - 1; i >= 0; --i)
 		{
-			const bool bBroadcast = TensorDesc.Sizes[i] < TargetDesc.Shape.Data[i];
+			const bool bBroadcast = TensorDesc.Sizes[i] < TargetDesc.GetShape().Data[i];
 
 			TensorDesc.Strides[i] = bBroadcast ? 0 : CurrStride;
 			CurrStride *= TensorDesc.Sizes[i];
 
-			TensorDesc.Sizes[i] = TargetDesc.Shape.Data[i];
+			TensorDesc.Sizes[i] = TargetDesc.GetShape().Data[i];
 		}
 	}
 
-	inline bool IsSameShape(const FMLTensorDesc& Left, const FMLTensorDesc& Right)
+	inline bool IsSameShape(const FTensor& Left, const FTensor& Right)
 	{
-		if (Left.Shape.Num() != Right.Shape.Num())
+		if (Left.GetShape().Rank() != Right.GetShape().Rank())
 		{
 			return false;
 		}
 		
-		for (int32 Idx = 0; Idx < Left.Shape.Num(); ++Idx)
+		for (int32 Idx = 0; Idx < Left.GetShape().Rank(); ++Idx)
 		{
-			if (Left.Shape.Data[Idx] != Right.Shape.Data[Idx])
+			if (Left.GetShape().Data[Idx] != Right.GetShape().Data[Idx])
 			{
 				return false;
 			}
@@ -270,7 +270,7 @@ namespace DmlUtil
 		}
 	}
 
-	inline uint64 CalculateBufferSize(const FTensorDesc& DmlTensor, const FMLTensorDesc& Desc)
+	inline uint64 CalculateBufferSize(const FTensorDesc& DmlTensor, const FTensor& Desc)
 	{
 		uint64 ElemSizeInBytes = Desc.GetElemByteSize();
 		
@@ -327,15 +327,15 @@ class FMLOperatorDml : public FMLOperatorRDG
 {
 public:
 
-	virtual bool Initialize(FDeviceContextDml* DevCtx, TArrayView<const FMLTensorDesc> InputTensors, TArrayView<const FMLTensorDesc> OutputTensors, const UE::NNECore::FAttributeMap& Attributes) = 0;
+	virtual bool Initialize(FDeviceContextDml* DevCtx, TArrayView<const FTensor> InputTensors, TArrayView<const FTensor> OutputTensors, const UE::NNECore::FAttributeMap& Attributes) = 0;
 
 	virtual void Dispatch(FRDGBuilder& GraphBuilder, TArrayView<const FMLTensorBinding> InInputBindings, TArrayView<const FMLTensorBinding> OutOutputBindings) = 0;
 
 protected:
 
-	bool InitDmlTensorDesc(DmlUtil::FTensorDesc& DmlTensorDesc, const FMLTensorDesc& TensorDesc)
+	bool InitDmlTensorDesc(DmlUtil::FTensorDesc& DmlTensorDesc, const FTensor& TensorDesc)
 	{
-		DML_TENSOR_DATA_TYPE DmlDataType = DmlUtil::GetTensorDataType(TensorDesc.DataType);
+		DML_TENSOR_DATA_TYPE DmlDataType = DmlUtil::GetTensorDataType(TensorDesc.GetDataType());
 
 		if (DmlDataType == DML_TENSOR_DATA_TYPE_UNKNOWN)
 		{
@@ -345,7 +345,7 @@ protected:
 			return false;
 		}
 
-		DmlTensorDesc.Sizes = TensorDesc.Shape.Data;
+		DmlTensorDesc.Sizes = TensorDesc.GetShape().Data;
 		// TODO: Support tensor padding using strides defined in FMLTensorDesc
 		//DmlUtil::SetTensorStrides(DmlTensorDesc, TensorDesc.Strides);
 		
@@ -355,19 +355,19 @@ protected:
 
 		BuffDesc.DataType = DmlDataType;
 		BuffDesc.Flags = DML_TENSOR_FLAG_NONE;
-		BuffDesc.DimensionCount = TensorDesc.Shape.Num();
+		BuffDesc.DimensionCount = TensorDesc.GetShape().Rank();
 		BuffDesc.Sizes = DmlTensorDesc.Sizes.GetData();
 		BuffDesc.Strides = nullptr;
-		BuffDesc.TotalTensorSizeInBytes = TensorDesc.DataSize;
+		BuffDesc.TotalTensorSizeInBytes = TensorDesc.GetDataSize();
 
 		DmlTensorDesc.Desc = DML_TENSOR_DESC{ DML_TENSOR_TYPE_BUFFER, &DmlTensorDesc.BuffDesc };
 
 		return true;
 	}
 
-	bool InitDmlTensorDesc(DmlUtil::FTensorDesc& DmlTensorDesc, const FMLTensorDesc& TensorDesc, const FMLTensorDesc& BroadcastDesc)
+	bool InitDmlTensorDesc(DmlUtil::FTensorDesc& DmlTensorDesc, const FTensor& TensorDesc, const FTensor& BroadcastDesc)
 	{
-		DML_TENSOR_DATA_TYPE DmlDataType = DmlUtil::GetTensorDataType(TensorDesc.DataType);
+		DML_TENSOR_DATA_TYPE DmlDataType = DmlUtil::GetTensorDataType(TensorDesc.GetDataType());
 
 		if (DmlDataType == DML_TENSOR_DATA_TYPE_UNKNOWN)
 		{
@@ -379,10 +379,10 @@ protected:
 
 		if (DmlUtil::IsSameShape(TensorDesc, BroadcastDesc))
 		{
-			DmlTensorDesc.Sizes = TensorDesc.Shape.Data;
+			DmlTensorDesc.Sizes = TensorDesc.GetShape().Data;
 			DmlUtil::SetTensorStrides(DmlTensorDesc, TensorDesc);
 		}
-		else if (TensorDesc.Shape.Num() > BroadcastDesc.Shape.Num())
+		else if (TensorDesc.GetShape().Rank() > BroadcastDesc.GetShape().Rank())
 		{
 			return false;
 		}
@@ -549,15 +549,15 @@ public:
 	//
 	//
 	//
-	virtual bool Initialize(FDeviceContextDml* InDevCtx, TArrayView<const FMLTensorDesc> InputTensors, TArrayView<const FMLTensorDesc> OutputTensors, const UE::NNECore::FAttributeMap& Attributes) override
+	virtual bool Initialize(FDeviceContextDml* InDevCtx, TArrayView<const FTensor> InputTensors, TArrayView<const FTensor> OutputTensors, const UE::NNECore::FAttributeMap& Attributes) override
 	{
 		// TODO: Setup attributes
-		Num = InputTensors[0].Volume;
+		Num = InputTensors[0].GetVolume();
 
 		DevCtx = InDevCtx;
 
-		const FMLTensorDesc& InputTensorDesc = InputTensors[0];
-		const FMLTensorDesc& OutputTensorDesc = OutputTensors[0];
+		const FTensor& InputTensorDesc = InputTensors[0];
+		const FTensor& OutputTensorDesc = OutputTensors[0];
 
 		Alpha = Attributes.GetValueOrDefault(TEXT("alpha"), Alpha);
 		Beta = Attributes.GetValueOrDefault(TEXT("beta"), Beta);
@@ -799,16 +799,16 @@ public:
 	//
 	//
 	//
-	virtual bool Initialize(FDeviceContextDml* InDevCtx, TArrayView<const FMLTensorDesc> InputTensors, TArrayView<const FMLTensorDesc> OutputTensors, const UE::NNECore::FAttributeMap& Attributes) override
+	virtual bool Initialize(FDeviceContextDml* InDevCtx, TArrayView<const FTensor> InputTensors, TArrayView<const FTensor> OutputTensors, const UE::NNECore::FAttributeMap& Attributes) override
 	{
 		// TODO: Setup attributes
-		Num = OutputTensors[0].Volume;
+		Num = OutputTensors[0].GetVolume();
 
 		DevCtx = InDevCtx;
 
-		const FMLTensorDesc& InputATensorDesc = InputTensors[0];
-		const FMLTensorDesc& InputBTensorDesc = InputTensors[1];
-		const FMLTensorDesc& OutputTensorDesc = OutputTensors[0];
+		const FTensor& InputATensorDesc = InputTensors[0];
+		const FTensor& InputBTensorDesc = InputTensors[1];
+		const FTensor& OutputTensorDesc = OutputTensors[0];
 
 		// Initialize tensor descriptors
 		DmlUtil::FTensorDesc	DmlInputATensorDesc{};
@@ -984,7 +984,7 @@ protected:
 
 private:
 
-	FMLOperatorDml* OpCreate(const FString& Name, TArrayView<const FMLTensorDesc> InputTensorDesc, TArrayView<const FMLTensorDesc> OutputTensorDescs, const UE::NNECore::FAttributeMap& Attributes);
+	FMLOperatorDml* OpCreate(const FString& Name, TArrayView<const FTensor> InputTensorDesc, TArrayView<const FTensor> OutputTensorDescs, const UE::NNECore::FAttributeMap& Attributes);
 	
 	TArray<FMLOperatorDml*>		Operators;
 	FDeviceContextDml*			DevCtx;
@@ -1274,21 +1274,21 @@ bool FMLInferenceModelDml::Init(UMLInferenceModel* InModel, FDeviceContextDml* I
 	{
 		const FString TypeName = Format.Operators[Idx].TypeName;
 
-		TArray<FMLTensorDesc> OpInputTensors;
-		TArray<FMLTensorDesc> OpOutputTensors;
+		TArray<FTensor> OpInputTensors;
+		TArray<FTensor> OpOutputTensors;
 		UE::NNECore::FAttributeMap AttributeMap;
 
 		for (int32 InputTensorIndex : Format.Operators[Idx].InTensors)
 		{
-			FSymbolicTensorDesc SymbolicTensorDesc = AllSymbolicTensors[InputTensorIndex];
+			FTensorDesc SymbolicTensorDesc = AllSymbolicTensors[InputTensorIndex];
 			//TODO jira 168972: Handle dynamic tensor desc, op should init from symbolic shapes
-			OpInputTensors.Emplace(FMLTensorDesc::MakeFromSymbolic(SymbolicTensorDesc));
+			OpInputTensors.Emplace(FTensor::MakeFromSymbolicDesc(SymbolicTensorDesc));
 		}
 		for (int32 OutputTensorIndex : Format.Operators[Idx].OutTensors)
 		{
-			FSymbolicTensorDesc SymbolicTensorDesc = AllSymbolicTensors[OutputTensorIndex];
+			FTensorDesc SymbolicTensorDesc = AllSymbolicTensors[OutputTensorIndex];
 			//TODO jira 168972: Handle dynamic tensor desc, op should init from symbolic shapes
-			OpOutputTensors.Emplace(FMLTensorDesc::MakeFromSymbolic(SymbolicTensorDesc));
+			OpOutputTensors.Emplace(FTensor::MakeFromSymbolicDesc(SymbolicTensorDesc));
 		}
 		for (const FMLFormatAttributeDesc& Desc : Format.Operators[Idx].Attributes)
 		{
@@ -1320,7 +1320,7 @@ void FMLInferenceModelDml::AddDispatchOps_RenderThread(FRDGBuilder& GraphBuilder
 	checkCode(
 		for (int32 i = 0; i < AllTensorBindings.Num(); ++i)
 		{
-			check(AllTensors[i].DataSize == AllTensorBindings[i].SizeInBytes);
+			check(AllTensors[i].GetDataSize() == AllTensorBindings[i].SizeInBytes);
 		}
 	);
 
@@ -1354,7 +1354,7 @@ void FMLInferenceModelDml::AddDispatchOps_RenderThread(FRDGBuilder& GraphBuilder
 //
 // Create operator
 //
-FMLOperatorDml* FMLInferenceModelDml::OpCreate(const FString& OpName, TArrayView<const FMLTensorDesc> InputTensorDescs, TArrayView<const FMLTensorDesc> OutputTensorDescs, const UE::NNECore::FAttributeMap& Attributes)
+FMLOperatorDml* FMLInferenceModelDml::OpCreate(const FString& OpName, TArrayView<const FTensor> InputTensorDescs, TArrayView<const FTensor> OutputTensorDescs, const UE::NNECore::FAttributeMap& Attributes)
 {
 	// TODO: Check if D3D12 device is valid
 
diff --git a/Engine/Plugins/Experimental/NNX/Source/NNXRuntimeRDG/Private/NNXRuntimeHLSL.cpp b/Engine/Plugins/Experimental/NNX/Source/NNXRuntimeRDG/Private/NNXRuntimeHLSL.cpp
--- a/Engine/Plugins/Experimental/NNX/Source/NNXRuntimeRDG/Private/NNXRuntimeHLSL.cpp
+++ b/Engine/Plugins/Experimental/NNX/Source/NNXRuntimeRDG/Private/NNXRuntimeHLSL.cpp
@@ -38,7 +38,7 @@ protected:
 
 private:
 
-	FMLOperatorHlsl* OpCreate(const FString& Name, TArrayView<const FMLTensorDesc> InputTensorDesc, TArrayView<const FMLTensorDesc> OutputTensorDescs, const UE::NNECore::FAttributeMap& Attributes);
+	FMLOperatorHlsl* OpCreate(const FString& Name, TArrayView<const FTensor> InputTensorDesc, TArrayView<const FTensor> OutputTensorDescs, const UE::NNECore::FAttributeMap& Attributes);
 
 	TArray<FMLOperatorHlsl*>	Operators;
 };
@@ -77,21 +77,21 @@ bool FMLInferenceModelHlsl::Init(UMLInferenceModel* InModel)
 	{
 		const FString TypeName = Format.Operators[Idx].TypeName;
 
-		TArray<FMLTensorDesc> OpInputTensors;
-		TArray<FMLTensorDesc> OpOutputTensors;
+		TArray<FTensor> OpInputTensors;
+		TArray<FTensor> OpOutputTensors;
 		UE::NNECore::FAttributeMap AttributeMap;
 
 		for (int32 InputTensorIndex : Format.Operators[Idx].InTensors)
 		{
-			FSymbolicTensorDesc SymbolicTensorDesc = AllSymbolicTensors[InputTensorIndex];
+			FTensorDesc SymbolicTensorDesc = AllSymbolicTensors[InputTensorIndex];
 			//TODO jira 168972: Handle dynamic tensor desc, op should init from symbolic shapes
-			OpInputTensors.Emplace(FMLTensorDesc::MakeFromSymbolic(SymbolicTensorDesc));
+			OpInputTensors.Emplace(FTensor::MakeFromSymbolicDesc(SymbolicTensorDesc));
 		}
 		for (int32 OutputTensorIndex : Format.Operators[Idx].OutTensors)
 		{
-			FSymbolicTensorDesc SymbolicTensorDesc = AllSymbolicTensors[OutputTensorIndex];
+			FTensorDesc SymbolicTensorDesc = AllSymbolicTensors[OutputTensorIndex];
 			//TODO jira 168972: Handle dynamic tensor desc, op should init from symbolic shapes
-			OpOutputTensors.Emplace(FMLTensorDesc::MakeFromSymbolic(SymbolicTensorDesc));
+			OpOutputTensors.Emplace(FTensor::MakeFromSymbolicDesc(SymbolicTensorDesc));
 		}
 		for (const FMLFormatAttributeDesc& Desc : Format.Operators[Idx].Attributes)
 		{
@@ -117,7 +117,7 @@ bool FMLInferenceModelHlsl::Init(UMLInferenceModel* InModel)
 //
 //
 //
-FMLOperatorHlsl* FMLInferenceModelHlsl::OpCreate(const FString& OpName, TArrayView<const FMLTensorDesc> InputTensorDescs, TArrayView<const FMLTensorDesc> OutputTensorDescs, const UE::NNECore::FAttributeMap& AttributeMap)
+FMLOperatorHlsl* FMLInferenceModelHlsl::OpCreate(const FString& OpName, TArrayView<const FTensor> InputTensorDescs, TArrayView<const FTensor> OutputTensorDescs, const UE::NNECore::FAttributeMap& AttributeMap)
 {
 	FMLOperatorRegistryHlsl::OperatorCreateFunc CreateFn = FMLOperatorRegistryHlsl::Get()->OpFind(OpName);
 
@@ -143,7 +143,7 @@ void FMLInferenceModelHlsl::AddDispatchOps_RenderThread(FRDGBuilder& GraphBuilde
 	checkCode(
 		for (int32 i = 0; i < AllTensorBindings.Num(); ++i)
 		{
-			check(AllTensors[i].DataSize == AllTensorBindings[i].SizeInBytes);
+			check(AllTensors[i].GetDataSize() == AllTensorBindings[i].SizeInBytes);
 		}
 	);
 
diff --git a/Engine/Plugins/Experimental/NNX/Source/NNXRuntimeRDG/Private/NNXRuntimeHLSLHelper.cpp b/Engine/Plugins/Experimental/NNX/Source/NNXRuntimeRDG/Private/NNXRuntimeHLSLHelper.cpp
--- a/Engine/Plugins/Experimental/NNX/Source/NNXRuntimeRDG/Private/NNXRuntimeHLSLHelper.cpp
+++ b/Engine/Plugins/Experimental/NNX/Source/NNXRuntimeRDG/Private/NNXRuntimeHLSLHelper.cpp
@@ -6,14 +6,14 @@
 
 namespace NNX
 {
-	void FillTensorStrideShaderParameters(const FMLTensorDesc& TensorDesc, TStaticArray<FUintVector4, NXRT_TENSORSTRIDEINFO_MAX_NUM_DIMENSIONS, 16U>& OutShaderParam, int32 Idx, int32 TargetNumdimensionForBroadcast)
+	void FillTensorStrideShaderParameters(const FTensor& TensorDesc, TStaticArray<FUintVector4, NXRT_TENSORSTRIDEINFO_MAX_NUM_DIMENSIONS, 16U>& OutShaderParam, int32 Idx, int32 TargetNumdimensionForBroadcast)
 	{
 		if (TargetNumdimensionForBroadcast == -1)
 		{
-			TargetNumdimensionForBroadcast = TensorDesc.Shape.Num();
+			TargetNumdimensionForBroadcast = TensorDesc.GetShape().Rank();
 		}
-		checkf(TargetNumdimensionForBroadcast >= TensorDesc.Shape.Num(), TEXT("Can't broadcast tensor from rank %d to rank %d, should be inferior or equal."), TensorDesc.Shape.Num(), TargetNumdimensionForBroadcast);
-		int32 Offset = TargetNumdimensionForBroadcast - TensorDesc.Shape.Num();
+		checkf(TargetNumdimensionForBroadcast >= TensorDesc.GetShape().Rank(), TEXT("Can't broadcast tensor from rank %d to rank %d, should be inferior or equal."), TensorDesc.GetShape().Rank(), TargetNumdimensionForBroadcast);
+		int32 Offset = TargetNumdimensionForBroadcast - TensorDesc.GetShape().Rank();
 
 		static_assert(FTensorShape::MaxRank <= NXRT_TENSORSTRIDEINFO_MAX_NUM_DIMENSIONS);
 		for (int32 i = 7; i >= 0; --i)
@@ -28,20 +28,20 @@ namespace NNX
 			}
 			else
 			{
-				OutShaderParam[i][Idx] = OutShaderParam[i + 1][Idx] * TensorDesc.Shape.Data[i + 1 - Offset];
+				OutShaderParam[i][Idx] = OutShaderParam[i + 1][Idx] * TensorDesc.GetShape().Data[i + 1 - Offset];
 			}
 		}
 	}
 
-	void FillTensorStrideForBroadcastShaderParameters(const FMLTensorDesc& TensorDesc, int32 OutputNumdimension, TStaticArray<FUintVector4, NXRT_TENSORSTRIDEINFO_MAX_NUM_DIMENSIONS, 16U>& OutShaderParam, int32 Idx)
+	void FillTensorStrideForBroadcastShaderParameters(const FTensor& TensorDesc, int32 OutputNumdimension, TStaticArray<FUintVector4, NXRT_TENSORSTRIDEINFO_MAX_NUM_DIMENSIONS, 16U>& OutShaderParam, int32 Idx)
 	{
-		checkf(OutputNumdimension >= TensorDesc.Shape.Num(), TEXT("Can't broadcast tensor from rank %d to rank %d, should be inferior or equal."), TensorDesc.Shape.Num(), OutputNumdimension);
+		checkf(OutputNumdimension >= TensorDesc.GetShape().Rank(), TEXT("Can't broadcast tensor from rank %d to rank %d, should be inferior or equal."), TensorDesc.GetShape().Rank(), OutputNumdimension);
 		FillTensorStrideShaderParameters(TensorDesc, OutShaderParam, Idx, OutputNumdimension);
-		int32 Offset = OutputNumdimension - TensorDesc.Shape.Num();
+		int32 Offset = OutputNumdimension - TensorDesc.GetShape().Rank();
 		for (int32 i = Offset; i < OutputNumdimension; ++i)
 		{
 			// the stride for broadcast dimension is kept as 0
-			if (TensorDesc.Shape.Data[i - Offset] == 1)
+			if (TensorDesc.GetShape().Data[i - Offset] == 1)
 			{
 				OutShaderParam[i][Idx] = 0;
 			}
diff --git a/Engine/Plugins/Experimental/NNX/Source/NNXRuntimeRDG/Private/NNXRuntimeHLSLHelper.h b/Engine/Plugins/Experimental/NNX/Source/NNXRuntimeRDG/Private/NNXRuntimeHLSLHelper.h
--- a/Engine/Plugins/Experimental/NNX/Source/NNXRuntimeRDG/Private/NNXRuntimeHLSLHelper.h
+++ b/Engine/Plugins/Experimental/NNX/Source/NNXRuntimeRDG/Private/NNXRuntimeHLSLHelper.h
@@ -8,7 +8,7 @@
 
 namespace NNX
 {
-	void FillTensorStrideShaderParameters(const FMLTensorDesc& TensorDesc, TStaticArray<FUintVector4, NXRT_TENSORSTRIDEINFO_MAX_NUM_DIMENSIONS, 16U>& OutShaderParam, int32 Idx, int32 TargetNumdimensionForBroadcast = -1);
-	void FillTensorStrideForBroadcastShaderParameters(const FMLTensorDesc& TensorDesc, int32 OutputNumdimension, TStaticArray<FUintVector4, NXRT_TENSORSTRIDEINFO_MAX_NUM_DIMENSIONS, 16U>& OutShaderParam, int32 Idx);
+	void FillTensorStrideShaderParameters(const FTensor& TensorDesc, TStaticArray<FUintVector4, NXRT_TENSORSTRIDEINFO_MAX_NUM_DIMENSIONS, 16U>& OutShaderParam, int32 Idx, int32 TargetNumdimensionForBroadcast = -1);
+	void FillTensorStrideForBroadcastShaderParameters(const FTensor& TensorDesc, int32 OutputNumdimension, TStaticArray<FUintVector4, NXRT_TENSORSTRIDEINFO_MAX_NUM_DIMENSIONS, 16U>& OutShaderParam, int32 Idx);
 	FIntVector ComputeElementWiseThreadGroups(uint32 ElementCount, uint32 GroupSizeX);
 } // NNX
\ No newline at end of file
diff --git a/Engine/Plugins/Experimental/NNX/Source/NNXRuntimeRDG/Private/NNXRuntimeHLSLOp.h b/Engine/Plugins/Experimental/NNX/Source/NNXRuntimeRDG/Private/NNXRuntimeHLSLOp.h
--- a/Engine/Plugins/Experimental/NNX/Source/NNXRuntimeRDG/Private/NNXRuntimeHLSLOp.h
+++ b/Engine/Plugins/Experimental/NNX/Source/NNXRuntimeRDG/Private/NNXRuntimeHLSLOp.h
@@ -13,7 +13,7 @@ namespace NNX
 	struct FMLOperatorHlsl : public FMLOperatorRDG
 	{
 		virtual ~FMLOperatorHlsl() = default;
-		virtual bool Initialize(TArrayView<const FMLTensorDesc> InputTensors, TArrayView<const FMLTensorDesc> OutputTensors, const UE::NNECore::FAttributeMap& Attributes) = 0;
+		virtual bool Initialize(TArrayView<const FTensor> InputTensors, TArrayView<const FTensor> OutputTensors, const UE::NNECore::FAttributeMap& Attributes) = 0;
 		virtual void Dispatch(FRDGBuilder& GraphBuilder, TArrayView<const FMLTensorBinding> InInputBindings, TArrayView<const FMLTensorBinding> OutOutputBindings) = 0;
 	};
 
diff --git a/Engine/Plugins/Experimental/NNX/Source/NNXRuntimeRDG/Private/NNXRuntimeRDG.cpp b/Engine/Plugins/Experimental/NNX/Source/NNXRuntimeRDG/Private/NNXRuntimeRDG.cpp
--- a/Engine/Plugins/Experimental/NNX/Source/NNXRuntimeRDG/Private/NNXRuntimeRDG.cpp
+++ b/Engine/Plugins/Experimental/NNX/Source/NNXRuntimeRDG/Private/NNXRuntimeRDG.cpp
@@ -18,7 +18,7 @@ namespace NNX
 bool AlwaysValidValidationFunction(
 	const UE::NNECore::FAttributeMap& AttributeMap, 
 	TConstArrayView<EMLTensorDataType> InputTensorTypes,
-	TConstArrayView<const FSymbolicTensorShape> InputShapes)
+	TConstArrayView<FSymbolicTensorShape> InputShapes)
 {
 	return true;
 }
@@ -211,7 +211,7 @@ bool FMLInferenceModelRDG::LoadModel(const FNNIModelRaw& InModel, FMLRuntimeForm
 		const FMLFormatTensorDesc& FormatTensorDesc = Format.Tensors[Idx];
 
 		FSymbolicTensorShape SymbolicShape = FSymbolicTensorShape::Make(FormatTensorDesc.Shape);
-		FSymbolicTensorDesc SymbolicTensor = FSymbolicTensorDesc::Make(FormatTensorDesc.Name, SymbolicShape, FormatTensorDesc.DataType);
+		FTensorDesc SymbolicTensor = FTensorDesc::Make(FormatTensorDesc.Name, SymbolicShape, FormatTensorDesc.DataType);
 		AllSymbolicTensors.Emplace(SymbolicTensor);
 		
 		if (FormatTensorDesc.Type == EMLFormatTensorType::Input)
@@ -244,7 +244,7 @@ bool FMLInferenceModelRDG::LoadModel(const FNNIModelRaw& InModel, FMLRuntimeForm
 /**
  * Run the inference model (synchronous version)
  */
-int FMLInferenceModelRDG::Run(TConstArrayView<const FMLTensorBinding> InInputBindings, TConstArrayView<const FTensorShape> InInputShapes, TConstArrayView<const FMLTensorBinding> InOutputBindings)
+int FMLInferenceModelRDG::Run(TConstArrayView<FMLTensorBinding> InInputBindings, TConstArrayView<FTensorShape> InInputShapes, TConstArrayView<FMLTensorBinding> InOutputBindings)
 {
 	// Prepare the model for the concrete input shapes
 	int Res = SetInputShapes(InInputShapes);
@@ -305,7 +305,7 @@ int FMLInferenceModelRDG::Run(TConstArrayView<const FMLTensorBinding> InInputBin
 	return Res;
 }
 
-int FMLInferenceModelRDG::SetInputShapes(TConstArrayView<const FTensorShape> InInputShapes)
+int FMLInferenceModelRDG::SetInputShapes(TConstArrayView<FTensorShape> InInputShapes)
 {
 	//Verify input shape are valid for the model
 	if (FMLInferenceModel::SetInputShapes(InInputShapes) != 0)
@@ -315,11 +315,11 @@ int FMLInferenceModelRDG::SetInputShapes(TConstArrayView<const FTensorShape> InI
 
 	//Run shape inference
 	AllTensors.Empty();
-	for (FSymbolicTensorDesc SymbolicTensorDesc : AllSymbolicTensors)
+	for (FTensorDesc SymbolicTensorDesc : AllSymbolicTensors)
 	{
 		//Fake shape inference for now
 		check(SymbolicTensorDesc.IsConcrete());
-		FMLTensorDesc TensorDesc = FMLTensorDesc::MakeFromSymbolic(SymbolicTensorDesc);
+		FTensor TensorDesc = FTensor::MakeFromSymbolicDesc(SymbolicTensorDesc);
 		AllTensors.Emplace(TensorDesc);
 	}
 
@@ -340,11 +340,11 @@ int FMLInferenceModelRDG::SetInputShapes(TConstArrayView<const FTensorShape> InI
 	return 0;
 }
 
-FRDGBufferDesc CreateRDGBufferDescFromTensorDesc(const FMLTensorDesc& TensorDesc)
+FRDGBufferDesc CreateRDGBufferDescFromTensorDesc(const FTensor& TensorDesc)
 {
 	// FIXME: CreateStructuredDesc() creates a crash on VulkanRHI
-	//FRDGBufferDesc Desc = FRDGBufferDesc::CreateStructuredDesc(TensorDesc.GetElemByteSize(), TensorDesc.Volume);
-	FRDGBufferDesc Desc = FRDGBufferDesc::CreateBufferDesc(TensorDesc.GetElemByteSize(), TensorDesc.Volume);
+	//FRDGBufferDesc Desc = FRDGBufferDesc::CreateStructuredDesc(TensorDesc.GetElemByteSize(), TensorDesc.GetVolume());
+	FRDGBufferDesc Desc = FRDGBufferDesc::CreateBufferDesc(TensorDesc.GetElemByteSize(), TensorDesc.GetVolume());
 
 	return Desc;
 }
@@ -352,7 +352,7 @@ FRDGBufferDesc CreateRDGBufferDescFromTensorDesc(const FMLTensorDesc& TensorDesc
 /**
  * Enqueue operators to RDG, the caller will run the GraphBuilder.Execute()
  */
-int FMLInferenceModelRDG::EnqueueRDG(FRDGBuilder& RDGBuilder, TConstArrayView<const FMLTensorBinding> InInputBindings, TConstArrayView<const FTensorShape> InInputShapes, TConstArrayView<const FMLTensorBinding> InOutputBindings)
+int FMLInferenceModelRDG::EnqueueRDG(FRDGBuilder& RDGBuilder, TConstArrayView<FMLTensorBinding> InInputBindings, TConstArrayView<FTensorShape> InInputShapes, TConstArrayView<FMLTensorBinding> InOutputBindings)
 {
 	check(IsInRenderingThread());
 
@@ -406,10 +406,10 @@ int FMLInferenceModelRDG::EnqueueRDG(FRDGBuilder& RDGBuilder, TConstArrayView<co
 	//Create intermediate tensors bindings
 	for (int32 Idx : IntermediateTensorIndices)
 	{
-		const FMLTensorDesc& TensorDesc = AllTensors[Idx];
+		const FTensor& TensorDesc = AllTensors[Idx];
 		FRDGBufferDesc Desc = CreateRDGBufferDescFromTensorDesc(TensorDesc);
-		FRDGBufferRef TensorBuffer = RDGBuilder.CreateBuffer(Desc, *TensorDesc.Name, ERDGBufferFlags::None);
-		AllTensorBindings[Idx] = FMLTensorBinding::FromRDG(TensorBuffer, TensorDesc.DataSize);
+		FRDGBufferRef TensorBuffer = RDGBuilder.CreateBuffer(Desc, *TensorDesc.GetName(), ERDGBufferFlags::None);
+		AllTensorBindings[Idx] = FMLTensorBinding::FromRDG(TensorBuffer, TensorDesc.GetDataSize());
 	}
 
 	//Insert input tensors bindings
@@ -447,12 +447,12 @@ int FMLInferenceModelRDG::EnqueueRDG(FRDGBuilder& RDGBuilder, TConstArrayView<co
  * Process tensor bindings and check if we need to create RDG Buffer for CPU binding 
  * Returns 0 on success, or index of a tensor binding if the tensor binding type is not supported.
  */
-int FMLInferenceModelRDG::SetTensors(FRDGBuilder& GraphBuilder, FMLTensorBindingArray& OutBindings, FMLIntArray& OutIndices, TArrayView<const FMLTensorBinding> InBindings, TArrayView<const FMLTensorDesc> InTensors)
+int FMLInferenceModelRDG::SetTensors(FRDGBuilder& GraphBuilder, FMLTensorBindingArray& OutBindings, FMLIntArray& OutIndices, TArrayView<const FMLTensorBinding> InBindings, TArrayView<const FTensor> InTensors)
 {
 	for (int32 Idx = 0; Idx < InBindings.Num(); ++Idx)
 	{
 		const FMLTensorBinding& Binding = InBindings[Idx];
-		const FMLTensorDesc& TensorDesc = InTensors[Idx];
+		const FTensor& TensorDesc = InTensors[Idx];
 
 		if (Binding.BindingType == EMLTensorBindingDataType::CPUMemory)
 		{
@@ -461,9 +461,9 @@ int FMLInferenceModelRDG::SetTensors(FRDGBuilder& GraphBuilder, FMLTensorBinding
 			// FIXME: We should use BUF_SourceCopy for only output buffers (GPU readback)
 			Desc.Usage = EBufferUsageFlags(Desc.Usage | BUF_SourceCopy);
 
-			FRDGBufferRef TensorBuffer = GraphBuilder.CreateBuffer(Desc, *TensorDesc.Name, ERDGBufferFlags::None);
+			FRDGBufferRef TensorBuffer = GraphBuilder.CreateBuffer(Desc, *TensorDesc.GetName(), ERDGBufferFlags::None);
 			
-			OutBindings.Emplace(FMLTensorBinding::FromRDG(TensorBuffer, InTensors[Idx].DataSize));
+			OutBindings.Emplace(FMLTensorBinding::FromRDG(TensorBuffer, InTensors[Idx].GetDataSize()));
 			OutIndices.Add(Idx);
 		}
 		else if (Binding.BindingType == EMLTensorBindingDataType::RDGBuffer)
@@ -490,9 +490,9 @@ void FMLInferenceModelRDG::AddTensorUploads_RenderThread(FRDGBuilder& GraphBuild
 		const int32				TensorIdx = InUploadIndices[Idx];
 		const FMLTensorBinding& RDGBinding = InRDGBindings[TensorIdx];
 		const FMLTensorBinding& InBinding = InBindings[TensorIdx];
-		const FMLTensorDesc&	TensorDesc = InputTensors[TensorIdx];
+		const FTensor&	TensorDesc = InputTensors[TensorIdx];
 
-		GraphBuilder.QueueBufferUpload(RDGBinding.Buffer, InBinding.CpuMemory, TensorDesc.DataSize, ERDGInitialDataFlags::NoCopy);
+		GraphBuilder.QueueBufferUpload(RDGBinding.Buffer, InBinding.CpuMemory, TensorDesc.GetDataSize(), ERDGInitialDataFlags::NoCopy);
 	}
 }
 
@@ -506,14 +506,14 @@ void FMLInferenceModelRDG::AddTensorReadbacks_RenderThread(FRDGBuilder& GraphBui
 		const int32				TensorIdx = InReadbackIndices[Idx];
 		const FMLTensorBinding& RDGBinding = InRDGBindings[TensorIdx];
 		const FMLTensorBinding& Binding = InBindings[TensorIdx];
-		const FMLTensorDesc&	TensorDesc = OutputTensors[TensorIdx];
+		const FTensor&	TensorDesc = OutputTensors[TensorIdx];
 
 		FMLTensorReadbackParameters* TensorReadbackParams = GraphBuilder.AllocParameters<FMLTensorReadbackParameters>();
 
 		TensorReadbackParams->Buffer = RDGBinding.Buffer;
 
 		GraphBuilder.AddPass(
-			RDG_EVENT_NAME("FMLInferenceModelAddTensorReadback:%s", *TensorDesc.Name),
+			RDG_EVENT_NAME("FMLInferenceModelAddTensorReadback:%s", *TensorDesc.GetName()),
 			TensorReadbackParams,
 			ERDGPassFlags::Readback | ERDGPassFlags::NeverCull,
 			[this, Binding, TensorDesc, TensorReadbackParams](FRHICommandListImmediate& RHICmdList)
@@ -532,10 +532,10 @@ void FMLInferenceModelRDG::AddTensorReadbacks_RenderThread(FRDGBuilder& GraphBui
 					RHICmdList.SubmitCommandsHint();
 				}
 
-				Readback.RHI->EnqueueCopy(RHICmdList, OutputBuffer, TensorDesc.DataSize);
+				Readback.RHI->EnqueueCopy(RHICmdList, OutputBuffer, TensorDesc.GetDataSize());
 				Readback.CpuMemory = Binding.CpuMemory;
 				Readback.Offset = 0;
-				Readback.Size = TensorDesc.DataSize;
+				Readback.Size = TensorDesc.GetDataSize();
 			}
 		);
 	}
diff --git a/Engine/Plugins/Experimental/NNX/Source/NNXRuntimeRDG/Private/NNXRuntimeRDG.h b/Engine/Plugins/Experimental/NNX/Source/NNXRuntimeRDG/Private/NNXRuntimeRDG.h
--- a/Engine/Plugins/Experimental/NNX/Source/NNXRuntimeRDG/Private/NNXRuntimeRDG.h
+++ b/Engine/Plugins/Experimental/NNX/Source/NNXRuntimeRDG/Private/NNXRuntimeRDG.h
@@ -35,7 +35,7 @@ namespace NNX
 {
 
 struct FTensorBinding;
-struct FTensorDesc;
+class FTensorDesc;
 
 /** Base class for all ML operators running on the RDG */
 struct FMLOperatorRDG
@@ -63,9 +63,9 @@ public:
 
 	~FMLInferenceModelRDG();
 
-	virtual int SetInputShapes(TConstArrayView<const FTensorShape> InputShapes) override;
-	virtual int Run(TConstArrayView<const FMLTensorBinding> InputBindings, TConstArrayView<const FTensorShape> InputShapes, TConstArrayView<const FMLTensorBinding> OutputBindings) override;
-	virtual int EnqueueRDG(FRDGBuilder& RDGBuilder, TConstArrayView<const FMLTensorBinding> InputBindings, TConstArrayView<const FTensorShape> InputShapes, TConstArrayView<const FMLTensorBinding> OutputBindings) override;
+	virtual int SetInputShapes(TConstArrayView<FTensorShape> InputShapes) override;
+	virtual int Run(TConstArrayView<FMLTensorBinding> InputBindings, TConstArrayView<FTensorShape> InputShapes, TConstArrayView<FMLTensorBinding> OutputBindings) override;
+	virtual int EnqueueRDG(FRDGBuilder& RDGBuilder, TConstArrayView<FMLTensorBinding> InputBindings, TConstArrayView<FTensorShape> InputShapes, TConstArrayView<FMLTensorBinding> OutputBindings) override;
 
 protected:
 
@@ -73,8 +73,7 @@ protected:
 
 	bool LoadModel(const FNNIModelRaw& InModel, FMLRuntimeFormat& Format);
 
-	int CreateIntermediateTensors(FRDGBuilder& GraphBuilder, FMLTensorBindingArray& OutBindings, TConstArrayView<const FMLTensorDesc> InTensors);
-	int SetTensors(FRDGBuilder& GraphBuilder, FMLTensorBindingArray& OutRDGBindings, FMLIntArray& OutIndices, TArrayView<const FMLTensorBinding> InBindings, TArrayView<const FMLTensorDesc> InTensors);
+	int SetTensors(FRDGBuilder& GraphBuilder, FMLTensorBindingArray& OutRDGBindings, FMLIntArray& OutIndices, TArrayView<const FMLTensorBinding> InBindings, TArrayView<const FTensor> InTensors);
 
 	virtual void AddDispatchOps_RenderThread(FRDGBuilder& GraphBuilder) = 0;
 
@@ -82,10 +81,10 @@ protected:
 	virtual void AddTensorReadbacks_RenderThread(FRDGBuilder& GraphBuilder, TArrayView<const int32> InReadbackIndices, TArrayView<const FMLTensorBinding> InRDGBindings, TArrayView<const FMLTensorBinding> InBindings);
 
 	//Tensors
-	TArray<FMLTensorDesc>		AllTensors;
-	TArray<FSymbolicTensorDesc>	AllSymbolicTensors;
-	TArray<FMLTensorDesc>		InputTensors;
-	TArray<FMLTensorDesc>		OutputTensors;
+	TArray<FTensor>		AllTensors;
+	TArray<FTensorDesc>	AllSymbolicTensors;
+	TArray<FTensor>		InputTensors;
+	TArray<FTensor>		OutputTensors;
 
 	//Tensor indices for models
 	TArray<int32>				IntermediateTensorIndices;
@@ -103,7 +102,7 @@ protected:
 
 //TODO jira 167585 remove default validation and declare contract in all HLSL operator (see HLSL Gemm for current example)
 //TODO jira 167584 remove default validation and declare contract in all DML operator (see HLSL Gemm for current example)
-bool AlwaysValidValidationFunction(const UE::NNECore::FAttributeMap& AttributeMap, TConstArrayView<EMLTensorDataType> InputTypes, TConstArrayView<const FSymbolicTensorShape> InputShapes);
+bool AlwaysValidValidationFunction(const UE::NNECore::FAttributeMap& AttributeMap, TConstArrayView<EMLTensorDataType> InputTypes, TConstArrayView<FSymbolicTensorShape> InputShapes);
 
 class FInputValidator
 {
@@ -158,7 +157,7 @@ class TOperatorRegistryRDG
 public:
 
 	typedef TOperatorType* (*OperatorCreateFunc)();
-	typedef bool (*OperatorValidateFunc)(const UE::NNECore::FAttributeMap& AttributeMap, TConstArrayView<EMLTensorDataType> InputTypes, TConstArrayView<const FSymbolicTensorShape> InputShapes);
+	typedef bool (*OperatorValidateFunc)(const UE::NNECore::FAttributeMap& AttributeMap, TConstArrayView<EMLTensorDataType> InputTypes, TConstArrayView<FSymbolicTensorShape> InputShapes);
 
 	static TOperatorRegistryRDG* Get()
 	{
diff --git a/Engine/Plugins/Experimental/NNX/Source/NNXUtils/Private/NNXModelBuilderONNX.cpp b/Engine/Plugins/Experimental/NNX/Source/NNXUtils/Private/NNXModelBuilderONNX.cpp
--- a/Engine/Plugins/Experimental/NNX/Source/NNXUtils/Private/NNXModelBuilderONNX.cpp
+++ b/Engine/Plugins/Experimental/NNX/Source/NNXUtils/Private/NNXModelBuilderONNX.cpp
@@ -276,7 +276,7 @@ private:
 //
 //
 //
-NNXUTILS_API bool CreateONNXModelForOperator(const FString& OperatorName, TConstArrayView<FMLTensorDesc> InInputTensors, TConstArrayView<FMLTensorDesc> InOutputTensors, FNNIModelRaw& Model)
+NNXUTILS_API bool CreateONNXModelForOperator(const FString& OperatorName, TConstArrayView<FTensor> InInputTensors, TConstArrayView<FTensor> InOutputTensors, FNNIModelRaw& Model)
 {
 	UE::NNECore::FAttributeMap EmptyAttributeMap;
 	return CreateONNXModelForOperator(OperatorName, InInputTensors, InOutputTensors, EmptyAttributeMap, Model);
@@ -285,7 +285,7 @@ NNXUTILS_API bool CreateONNXModelForOperator(const FString& OperatorName, TConst
 //
 //
 //
-NNXUTILS_API bool CreateONNXModelForOperator(const FString& OperatorName, TConstArrayView<FMLTensorDesc> InInputTensors, TConstArrayView<FMLTensorDesc> InOutputTensors, const UE::NNECore::FAttributeMap& Attributes, FNNIModelRaw& Model)
+NNXUTILS_API bool CreateONNXModelForOperator(const FString& OperatorName, TConstArrayView<FTensor> InInputTensors, TConstArrayView<FTensor> InOutputTensors, const UE::NNECore::FAttributeMap& Attributes, FNNIModelRaw& Model)
 {
 	Model = FNNIModelRaw{};
 	
@@ -297,8 +297,8 @@ NNXUTILS_API bool CreateONNXModelForOperator(const FString& OperatorName, TConst
 	
 	for (int32 Idx = 0; Idx < InInputTensors.Num(); ++Idx)
 	{
-		const FMLTensorDesc& Desc = InInputTensors[Idx];
-		IMLModelBuilder::HTensor Tensor = Builder->AddTensor(Desc.Name, Desc.DataType, MakeArrayView((const int32*) Desc.Shape.Data.GetData(), Desc.Shape.Num()));
+		const FTensor& Desc = InInputTensors[Idx];
+		IMLModelBuilder::HTensor Tensor = Builder->AddTensor(Desc.GetName(), Desc.GetDataType(), MakeArrayView((const int32*)Desc.GetShape().Data.GetData(), Desc.GetShape().Rank()));
 
 		InputTensors.Emplace(Tensor);
 		Builder->AddInput(Tensor);
@@ -308,8 +308,8 @@ NNXUTILS_API bool CreateONNXModelForOperator(const FString& OperatorName, TConst
 
 	for (int32 Idx = 0; Idx < InOutputTensors.Num(); ++Idx)
 	{
-		const FMLTensorDesc& Desc = InOutputTensors[Idx];
-		IMLModelBuilder::HTensor Tensor = Builder->AddTensor(Desc.Name, Desc.DataType, MakeArrayView((const int32*) Desc.Shape.Data.GetData(), Desc.Shape.Num()));
+		const FTensor& Desc = InOutputTensors[Idx];
+		IMLModelBuilder::HTensor Tensor = Builder->AddTensor(Desc.GetName(), Desc.GetDataType(), MakeArrayView((const int32*)Desc.GetShape().Data.GetData(), Desc.GetShape().Rank()));
 
 		OutputTensors.Emplace(Tensor);
 		Builder->AddOutput(Tensor);
diff --git a/Engine/Plugins/Experimental/NNX/Source/NNXUtils/Private/NNXUtilsModule.cpp b/Engine/Plugins/Experimental/NNX/Source/NNXUtils/Private/NNXUtilsModule.cpp
--- a/Engine/Plugins/Experimental/NNX/Source/NNXUtils/Private/NNXUtilsModule.cpp
+++ b/Engine/Plugins/Experimental/NNX/Source/NNXUtils/Private/NNXUtilsModule.cpp
@@ -91,8 +91,8 @@ static FAutoConsoleCommand ConsoleCmdNNXOptimizeModel(
 			}
 			else if (!OpName.IsEmpty())
 			{
-				FMLTensorDesc InputTensor = FMLTensorDesc::Make(TEXT("in"), FTensorShape::Make({1, 512 }), EMLTensorDataType::Float);
-				FMLTensorDesc OutputTensor = FMLTensorDesc::Make(TEXT("out"), FTensorShape::Make({ 1, 512 }), EMLTensorDataType::Float);
+				FTensor InputTensor = FTensor::Make(TEXT("in"), FTensorShape::Make({1, 512 }), EMLTensorDataType::Float);
+				FTensor OutputTensor = FTensor::Make(TEXT("out"), FTensorShape::Make({ 1, 512 }), EMLTensorDataType::Float);
 
 				if (!NNX::CreateONNXModelForOperator(OpName, MakeArrayView(&InputTensor, 1), MakeArrayView(&OutputTensor, 1), ONNXModel))
 				{
@@ -135,13 +135,13 @@ static FAutoConsoleCommand ConsoleCmdNNXCreateModel(
 	FConsoleCommandWithArgsDelegate::CreateStatic(
 		[](const TArray<FString>& Args)
 		{
-			using NNX::FMLTensorDesc;
+			using NNX::FTensor;
 
 			if (Args.Num() > 0)
 			{
 				FNNIModelRaw ONNXModel;
-				FMLTensorDesc InputTensor = FMLTensorDesc::Make(TEXT("in"), FTensorShape::Make({ 1, 512 }), EMLTensorDataType::Float);
-				FMLTensorDesc OutputTensor = FMLTensorDesc::Make(TEXT("out"), FTensorShape::Make({ 1, 512 }), EMLTensorDataType::Float);
+				FTensor InputTensor = FTensor::Make(TEXT("in"), FTensorShape::Make({ 1, 512 }), EMLTensorDataType::Float);
+				FTensor OutputTensor = FTensor::Make(TEXT("out"), FTensorShape::Make({ 1, 512 }), EMLTensorDataType::Float);
 
 				if (!NNX::CreateONNXModelForOperator(Args[0], MakeArrayView(&InputTensor, 1), MakeArrayView(&OutputTensor, 1), ONNXModel))
 				{
diff --git a/Engine/Plugins/Experimental/NNX/Source/NNXUtils/Public/NNXModelBuilder.h b/Engine/Plugins/Experimental/NNX/Source/NNXUtils/Public/NNXModelBuilder.h
--- a/Engine/Plugins/Experimental/NNX/Source/NNXUtils/Public/NNXModelBuilder.h
+++ b/Engine/Plugins/Experimental/NNX/Source/NNXUtils/Public/NNXModelBuilder.h
@@ -87,8 +87,8 @@ NNXUTILS_API IMLModelBuilder* CreateONNXModelBuilder(int64 IrVersion = OnnxIrVer
 /**
  * Utility functions to create single layer NN for operator testing with optional attributes
  */
-NNXUTILS_API bool CreateONNXModelForOperator(const FString& OperatorName, TConstArrayView<const FMLTensorDesc> InInputTensors, TConstArrayView<const FMLTensorDesc> InOutputTensors, FNNIModelRaw& ModelData);
-NNXUTILS_API bool CreateONNXModelForOperator(const FString& OperatorName, TConstArrayView<const FMLTensorDesc> InInputTensors, TConstArrayView<const FMLTensorDesc> InOutputTensors, const UE::NNECore::FAttributeMap& Attributes, FNNIModelRaw& ModelData);
+NNXUTILS_API bool CreateONNXModelForOperator(const FString& OperatorName, TConstArrayView<FTensor> InInputTensors, TConstArrayView<FTensor> InOutputTensors, FNNIModelRaw& ModelData);
+NNXUTILS_API bool CreateONNXModelForOperator(const FString& OperatorName, TConstArrayView<FTensor> InInputTensors, TConstArrayView<FTensor> InOutputTensors, const UE::NNECore::FAttributeMap& Attributes, FNNIModelRaw& ModelData);
 
 /**
  * Create an instance of NNX model builder that creates NNX model/format in memory
