Added support for DirectML Gemm operator

#rnx

diff --git a/Engine/Plugins/Experimental/NNX/Source/NNXQA/Resources/NNXQATestDesc.json b/Engine/Plugins/Experimental/NNX/Source/NNXQA/Resources/NNXQATestDesc.json
--- a/Engine/Plugins/Experimental/NNX/Source/NNXQA/Resources/NNXQATestDesc.json
+++ b/Engine/Plugins/Experimental/NNX/Source/NNXQA/Resources/NNXQATestDesc.json
@@ -236,7 +236,6 @@
 		},
 		{
 			"category": "Gemm",
-			"runtimes": [ { "name": "NNXRuntimeDml", "skip": true } ],
 			"targets": [
 				{ "target": "Gemm" }
 			]
diff --git a/Engine/Plugins/Experimental/NNX/Source/NNXRuntimeRDG/Private/NNXRuntimeDML.cpp b/Engine/Plugins/Experimental/NNX/Source/NNXRuntimeRDG/Private/NNXRuntimeDML.cpp
--- a/Engine/Plugins/Experimental/NNX/Source/NNXRuntimeRDG/Private/NNXRuntimeDML.cpp
+++ b/Engine/Plugins/Experimental/NNX/Source/NNXRuntimeRDG/Private/NNXRuntimeDML.cpp
@@ -161,8 +161,8 @@ namespace DmlUtil
 {
 	struct FTensorDesc
 	{
-		DML_BUFFER_TENSOR_DESC												BuffDesc;
-		DML_TENSOR_DESC														Desc;
+		DML_BUFFER_TENSOR_DESC										BuffDesc;
+		DML_TENSOR_DESC												Desc;
 		TArray<uint32, TInlineAllocator<FTensorShape::MaxRank>>		Sizes;
 		TArray<uint32, TInlineAllocator<FTensorShape::MaxRank>>		Strides;
 	};
@@ -966,6 +966,198 @@ public:
 	}
 };
 
+/**
+ * Gemm
+ */
+class FMLOperatorDmlGemm : public FMLOperatorDml
+{
+public:
+
+	static FMLOperatorDml* Create()
+	{
+		return new FMLOperatorDmlGemm();
+	}
+
+	//
+	//
+	//
+	virtual bool Initialize(FDeviceContextDml* InDevCtx, TArrayView<const FTensor> InputTensors, TArrayView<const FTensor> OutputTensors, const UE::NNECore::FAttributeMap& Attributes) override
+	{
+		// Setup attributes
+		float Alpha = 1.0f;
+		float Beta = 1.0f;
+		int32 TransA = 0;
+		int32 TransB = 0;
+		
+		Alpha = Attributes.GetValueOrDefault(TEXT("alpha"), Alpha);
+		Beta = Attributes.GetValueOrDefault(TEXT("beta"), Beta);
+		TransA = Attributes.GetValueOrDefault(TEXT("transA"), TransA);
+		TransB = Attributes.GetValueOrDefault(TEXT("transB"), TransB);
+		
+		DevCtx = InDevCtx;
+
+		const FTensor& InputATensorDesc = InputTensors[0];
+		const FTensor& InputBTensorDesc = InputTensors[1];
+		const FTensor& OutputTensorDesc = OutputTensors[0];
+
+		// Initialize tensor descriptors
+		DmlUtil::FTensorDesc	DmlInputATensorDesc{};
+		DmlUtil::FTensorDesc	DmlInputBTensorDesc{};
+		DmlUtil::FTensorDesc	DmlInputCTensorDesc{};
+		DmlUtil::FTensorDesc	DmlOutputTensorDesc{};
+
+		if (!InitDmlTensorDesc(DmlInputATensorDesc, InputATensorDesc))
+		{
+			UE_LOG(LogNNX, Warning, TEXT("Failed to initialize tensor(s) for DML inference"));
+			return false;
+		}
+
+		if (!InitDmlTensorDesc(DmlInputBTensorDesc, InputBTensorDesc))
+		{
+			UE_LOG(LogNNX, Warning, TEXT("Failed to initialize tensor(s) for DML inference"));
+			return false;
+		}
+
+		if (InputTensors.Num() > 2)
+		{
+			const FTensor& InputCTensorDesc = InputTensors[2];
+
+			if (!InitDmlTensorDesc(DmlInputCTensorDesc, InputCTensorDesc, OutputTensorDesc))
+			{
+				UE_LOG(LogNNX, Warning, TEXT("Failed to initialize tensor(s) for DML inference"));
+				return false;
+			}
+		}
+
+		if (!InitDmlTensorDesc(DmlOutputTensorDesc, OutputTensorDesc))
+		{
+			UE_LOG(LogNNX, Warning, TEXT("Failed to initialize tensor(s) for DML inference"));
+			return false;
+		}
+
+		DML_GEMM_OPERATOR_DESC	DmlGemmOpDesc{};
+
+		DmlGemmOpDesc.ATensor = &DmlInputATensorDesc.Desc;
+		DmlGemmOpDesc.BTensor = &DmlInputBTensorDesc.Desc;
+		DmlGemmOpDesc.CTensor = InputTensors.Num() > 2 ? &DmlInputCTensorDesc.Desc : nullptr;
+		DmlGemmOpDesc.OutputTensor = &DmlOutputTensorDesc.Desc;
+		DmlGemmOpDesc.Alpha = Alpha;
+		DmlGemmOpDesc.Beta = Beta;
+		DmlGemmOpDesc.TransA = TransA ? DML_MATRIX_TRANSFORM_TRANSPOSE : DML_MATRIX_TRANSFORM_NONE;
+		DmlGemmOpDesc.TransB = TransB ? DML_MATRIX_TRANSFORM_TRANSPOSE : DML_MATRIX_TRANSFORM_NONE;
+
+		DML_OPERATOR_DESC DmlOpDesc{};
+
+		DmlOpDesc.Type = DML_OPERATOR_GEMM;
+		DmlOpDesc.Desc = &DmlGemmOpDesc;
+
+		if (!CompileOperator(DmlOpDesc))
+		{
+			UE_LOG(LogNNX, Warning, TEXT("Failed to compile DML operator"));
+			return false;
+		}
+
+		return true;
+	}
+
+	//
+	//
+	//
+	virtual void Dispatch(FRDGBuilder& GraphBuilder, TArrayView<const FMLTensorBinding> InInputBindings, TArrayView<const FMLTensorBinding> OutOutputBindings) override
+	{
+		ID3D12DynamicRHI* DynamicRHI = GetID3D12DynamicRHI();
+
+		const int32 NumInputs = InInputBindings.Num();
+		const bool bIsUsingBias = NumInputs > 2;
+
+		FMLGemmParameters* Params = GraphBuilder.AllocParameters<FMLGemmParameters>();
+
+		Params->A = InInputBindings[0].Buffer;
+		Params->B = InInputBindings[1].Buffer;
+		Params->C = bIsUsingBias ? InInputBindings[2].Buffer : nullptr;
+		Params->Y = OutOutputBindings[0].Buffer;
+
+		GraphBuilder.AddPass(
+			RDG_EVENT_NAME("FMLGemmDml_Transition"),
+			Params,
+			ERDGPassFlags::Compute | ERDGPassFlags::NeverCull,
+			[Params](FRHICommandListImmediate& RHICmdList)
+			{
+				FRHIBuffer* InputABuffer = Params->A->GetRHI();
+				FRHIBuffer* InputBBuffer = Params->B->GetRHI();
+				FRHIBuffer* InputCBuffer = Params->C ? Params->C->GetRHI() : nullptr;
+				FRHIBuffer* OutputBuffer = Params->Y->GetRHI();
+
+				FRHITransitionInfo Transitions[] =
+				{
+					FRHITransitionInfo(InputABuffer, ERHIAccess::Unknown, ERHIAccess::UAVCompute),
+					FRHITransitionInfo(InputBBuffer, ERHIAccess::Unknown, ERHIAccess::UAVCompute),
+					FRHITransitionInfo(OutputBuffer, ERHIAccess::Unknown, ERHIAccess::UAVCompute),
+					FRHITransitionInfo(InputCBuffer, ERHIAccess::Unknown, ERHIAccess::UAVCompute)
+				};
+
+				RHICmdList.Transition(MakeArrayView(Transitions, Params->C ? 4 : 3));
+
+				// FIXME: We need to flush commands here to transition the resources manually
+				RHICmdList.SubmitCommandsHint();
+			});
+
+		GraphBuilder.AddPass(
+			RDG_EVENT_NAME("FMLGemmDml_Dispatch"),
+			Params,
+			ERDGPassFlags::Compute | ERDGPassFlags::NeverCull,
+			[this, DynamicRHI, Params](FRHICommandListImmediate& RHICmdList)
+			{
+				FRHIBuffer* InputBufferA = Params->A->GetRHI();
+				FRHIBuffer* InputBufferB = Params->B->GetRHI();
+				FRHIBuffer* InputBufferC = Params->C ? Params->C->GetRHI() : nullptr;
+				FRHIBuffer* OutputBuffer = Params->Y->GetRHI();
+
+				// We need to defer the DML command list record
+				RHICmdList.EnqueueLambda(
+					[this, DynamicRHI, InputBufferA, InputBufferB, InputBufferC, OutputBuffer](FRHICommandListImmediate& RHICmdList)
+					{
+						DML_BINDING_TABLE_DESC DmlBindingTableDesc{};
+
+						ResetBindingTable(DmlBindingTableDesc);
+
+						// Get native resources for DML binding
+						ID3D12Resource* InputAResource = DynamicRHI->RHIGetResource(InputBufferA);
+						ID3D12Resource* InputBResource = DynamicRHI->RHIGetResource(InputBufferB);
+						ID3D12Resource* InputCResource = InputBufferC ? DynamicRHI->RHIGetResource(InputBufferC) : nullptr;
+						ID3D12Resource* OutputResource = DynamicRHI->RHIGetResource(OutputBuffer);
+
+						// TODO: Bind resources to DML binding table
+						DML_BUFFER_BINDING	InputABuffBind { InputAResource, 0, InputAResource->GetDesc().Width };
+						DML_BUFFER_BINDING	InputBBuffBind { InputBResource, 0, InputBResource->GetDesc().Width };
+						DML_BUFFER_BINDING	InputCBuffBind { InputCResource, 0, InputCResource ? InputCResource->GetDesc().Width : 0 };
+						DML_BUFFER_BINDING	OutputBuffBind { OutputResource, 0, OutputResource->GetDesc().Width };
+
+						DML_BINDING_DESC	InputABindDesc{ DML_BINDING_TYPE_BUFFER, &InputABuffBind };
+						DML_BINDING_DESC	InputBBindDesc{ DML_BINDING_TYPE_BUFFER, &InputBBuffBind };
+						DML_BINDING_DESC	InputCBindDesc{ DML_BINDING_TYPE_BUFFER, &InputCBuffBind };
+						DML_BINDING_DESC	OutputBindDesc{ DML_BINDING_TYPE_BUFFER, &OutputBuffBind };
+
+						DML_BINDING_DESC	InputBindings[] = { InputABindDesc, InputBBindDesc, InputCBindDesc };
+
+						BindingTable->BindInputs(3, InputBindings);
+						BindingTable->BindOutputs(1, &OutputBindDesc);
+
+						// Record command list
+						ID3D12GraphicsCommandList* CmdList = nullptr;
+
+						CmdList = DynamicRHI->RHIGetGraphicsCommandList(DevCtx->DeviceIndex);
+
+						CmdList->SetDescriptorHeaps(1, &DevCtx->DescHeap);
+						DevCtx->CmdRec->RecordDispatch(CmdList, CompiledOp, BindingTable);
+
+					}
+				);
+			}
+		);
+	}
+};
+
 //
 //
 //
@@ -1038,6 +1230,7 @@ private:
 
 	bool RegisterElementWiseUnaryOperators();
 	bool RegisterElementWiseBinaryOperators();
+	bool RegisterGemmOperator();
 	
 	FDeviceContextDml		Ctx;
 };
@@ -1090,6 +1283,7 @@ bool FMLRuntimeDml::Init()
 
 	RegisterElementWiseUnaryOperators();
 	RegisterElementWiseBinaryOperators();
+	RegisterGemmOperator();
 
 	Ctx.DeviceIndex = 0;
 	Ctx.D3D12Device = RHI->RHIGetDevice(Ctx.DeviceIndex);
@@ -1237,6 +1431,14 @@ bool FMLRuntimeDml::RegisterElementWiseBinaryOperators()
 	return true;
 }
 
+//
+//
+//
+bool FMLRuntimeDml::RegisterGemmOperator()
+{
+	FMLOperatorRegistryDml::Get()->OpAdd(TEXT("Gemm"), FMLOperatorDmlGemm::Create);
+	return true;
+}
 
 //
 //
