NNI Dynamic shape support for CPU and ORT engines

Problem:
* Models with dynamic input shapes are not supported

Solution:
* Added support for ORT based engines (including NNXRuntimeCPU) by
-> Make SetInputTensorShapes call mandatory, so it is explicit and controlable by customer code
-> For static model -> no change of behavior. Past SetInputTensorShapes ORT will output to customer allocated output buffer.
-> For dynamic model -> Allow ORT based implementation to run shape inference during model execution, allocating output tensor on the fly then copy back the result to customer allocated buffer

Tests:
* Added two test/models (with one and two layers respectively) with dynamic input and output shapes
* All NNX tests pass under win64 editor

diff --git a/Engine/Plugins/Experimental/NNX/Source/NNXCore/Private/NNXRuntime.cpp b/Engine/Plugins/Experimental/NNX/Source/NNXCore/Private/NNXRuntime.cpp
--- a/Engine/Plugins/Experimental/NNX/Source/NNXCore/Private/NNXRuntime.cpp
+++ b/Engine/Plugins/Experimental/NNX/Source/NNXCore/Private/NNXRuntime.cpp
@@ -42,8 +42,10 @@ namespace NNX
 		return Binding;
 	}
 
-	int FMLInferenceModel::SetInputShapes(TConstArrayView<FTensorShape> InInputShapes)
+	int FMLInferenceModel::SetInputTensorShapes(TConstArrayView<FTensorShape> InInputShapes)
 	{
+		InputTensorShapes.Empty();
+		
 		if (InInputShapes.Num() != InputSymbolicTensors.Num())
 		{
 			UE_LOG(LogNNX, Warning, TEXT("Number of input shapes does not match number of input tensors"));
@@ -60,9 +62,10 @@ namespace NNX
 			}
 		}
 
-		//Implementations are responsible to handle shape inference.
-		//This base implementation only validate that all inputs are
-		//matching what the model can support.
+		InputTensorShapes = InInputShapes;
+
+		//Implementations are responsible to handle output and intermediate tensor shape inference.
+		//This base implementation only validate that all inputs are matching what the model can support.
 		return 0;
 	}
 
diff --git a/Engine/Plugins/Experimental/NNX/Source/NNXCore/Public/NNXRuntime.h b/Engine/Plugins/Experimental/NNX/Source/NNXCore/Public/NNXRuntime.h
--- a/Engine/Plugins/Experimental/NNX/Source/NNXCore/Public/NNXRuntime.h
+++ b/Engine/Plugins/Experimental/NNX/Source/NNXCore/Public/NNXRuntime.h
@@ -106,25 +106,26 @@ public:
 	TConstArrayView<FTensorDesc> GetInputTensorDescs() const;
 	TConstArrayView<FTensorDesc> GetOutputTensorDescs() const;
 
-	/** Getters for input/output shapes if they were set already. Empty list otherwise. See SetInputShapes() */
+	/** Getters for input shapes if they were set already (see SetInputTensorShapes). Empty list otherwise. */
 	TConstArrayView<FTensorShape> GetInputTensorShapes() const;
+
+	/** Getters for outputs shapes if they were already resolved. Empty list otherwise.
+	* Output shape might be resolved after a call to SetInputTensorShapes() if the model and engine support it
+	* otherwise they will be resolved during Run() or EnqueueRDG() */
 	TConstArrayView<FTensorShape> GetOutputTensorShapes() const;
 	
-	/** This call will prepare the model to be run with the given input shape, it is an optional call.
-	 * Run() and EnqueueRDG() will call it internally if it was not call explicitely.
-	 * It is recommended to call this function ahead of time to avoid overhead during the first model execution. */
-	virtual int SetInputShapes(TConstArrayView<FTensorShape> InInputShapes);
+	/** Prepare the model to be run with the given input shape, The call is mandatory before Run or Enqueue() can be called. */
+	virtual int SetInputTensorShapes(TConstArrayView<FTensorShape> InInputShapes);
 
-	/** This call is synchronous on all engine types(CPU, RDG, GPU), i.e.the calling thread will be blocked.
-	 * until execution is finished */
+	/** This call is synchronous on all engine types(CPU, RDG, GPU), i.e.the calling thread will be blocked until execution is finished.
+	* Bindings should point to a buffers big enough */
 	virtual int Run(TConstArrayView<FMLTensorBinding> InInputTensors,
-		            TConstArrayView<FTensorShape> InInputShapes,
 					TConstArrayView<FMLTensorBinding> InOutputTensors) = 0;
 
-	/** Enqueue the execution on the Render graph render thread. It's caller's responsibility to actually run the graph. */
+	/** Enqueue the execution on the Render graph render thread. It's caller's responsibility to actually run the graph.
+	* Bindings should point to a buffers big enough */
 	virtual int EnqueueRDG(FRDGBuilder& RDGBuilder, 
 		TConstArrayView<FMLTensorBinding> InInputTensors, 
-		TConstArrayView<FTensorShape> InInputShapes,
 		TConstArrayView<FMLTensorBinding> InOutputTensors)
 	{
 		return -1;
@@ -136,10 +137,10 @@ protected:
 		: Type(InType)
 	{}
 
-	TArray<FTensorShape>		InputTensorShapes;
-	TArray<FTensorShape>		OutputTensorShapes;
-	TArray<FTensorDesc>	InputSymbolicTensors;
-	TArray<FTensorDesc>	OutputSymbolicTensors;
+	TArray<FTensorShape>	InputTensorShapes;
+	TArray<FTensorShape>	OutputTensorShapes;
+	TArray<FTensorDesc>		InputSymbolicTensors;
+	TArray<FTensorDesc>		OutputSymbolicTensors;
 	EMLInferenceModelType	Type;
 };
 
diff --git a/Engine/Plugins/Experimental/NNX/Source/NNXQA/Private/NNXQAUtils.cpp b/Engine/Plugins/Experimental/NNX/Source/NNXQA/Private/NNXQAUtils.cpp
--- a/Engine/Plugins/Experimental/NNX/Source/NNXQA/Private/NNXQAUtils.cpp
+++ b/Engine/Plugins/Experimental/NNX/Source/NNXQA/Private/NNXQAUtils.cpp
@@ -14,6 +14,21 @@ namespace NNX
 {
 namespace Test 
 {
+	static void FillTensors(
+		TConstArrayView<FTensor> TensorsFromTestSetup,
+		TConstArrayView<FTensorDesc> TensorDescsFromModel,
+		TArray<FTensor>& OutTensors)
+	{
+		OutTensors = TensorsFromTestSetup;
+		if (OutTensors.IsEmpty())
+		{
+			for (const FTensorDesc& SymbolicTensorDesc : TensorDescsFromModel)
+			{
+				FTensor TensorDesc = FTensor::MakeFromSymbolicDesc(SymbolicTensorDesc);
+				OutTensors.Emplace(TensorDesc);
+			}
+		}
+	}
 	static void FillTensorBindings(TConstArrayView<FTensor> TensorDescs,
 		TArray<TArray<char>>& OutMemBuffers, TArray<FMLTensorBinding>& OutBindings,
         std::function<float(EMLTensorDataType, uint32, uint32)> Initializer)
@@ -140,8 +155,8 @@ namespace Test
 	}
 
 	template<typename T> bool CompareTensorData(
-		const NNX::FTensor& RefTensorDesc,   const TArray<char>& RefRawBuffer,
-		const NNX::FTensor& OtherTensorDesc, const TArray<char>& OtherRawBuffer,
+		const FTensor& RefTensorDesc,   const TArray<char>& RefRawBuffer,
+		const FTensor& OtherTensorDesc, const TArray<char>& OtherRawBuffer,
 		float AbsoluteErrorEpsilon, float RelativeErrorPercent)
 	{
 		const T* RefBuffer = (T*)RefRawBuffer.GetData();
@@ -242,8 +257,8 @@ namespace Test
 	}
 
 	bool VerifyTensorResult(
-		const NNX::FTensor& RefTensorDesc, const TArray<char>& RefRawBuffer,
-		const NNX::FTensor& OtherTensorDesc, const TArray<char>& OtherRawBuffer,
+		const FTensor& RefTensorDesc, const TArray<char>& RefRawBuffer,
+		const FTensor& OtherTensorDesc, const TArray<char>& OtherRawBuffer,
 		float AbsoluteErrorEpsilon, float RelativeErrorPercent)
 	{
 		bool bTensorDescMatch = true;
@@ -325,11 +340,11 @@ namespace Test
 		return FMath::Sin((float)(ElementIndex + IndexOffsetBetweenTensor * TensorIndex));
 	}
 
-	static int RunTestInference(const FNNIModelRaw& ONNXModelData, IRuntime* Runtime, 
-		TArray<FTensor>& OutOutputTensorsDesc, TArray<TArray<char>>& OutOutputMemBuffers)
+	static int RunTestInference(const FNNIModelRaw& ONNXModelData, const FTests::FTestSetup& TestSetup, 
+		IRuntime* Runtime, TArray<FTensor>& OutOutputTensors, TArray<TArray<char>>& OutOutputMemBuffers)
 	{
 		OutOutputMemBuffers.Empty();
-		OutOutputTensorsDesc.Empty();
+		OutOutputTensors.Empty();
 
 		FOptimizerOptionsMap Options;
 		FNNIModelRaw RuntimeModelData;
@@ -350,37 +365,25 @@ namespace Test
 			return -1;
 		}
 
-		//bind input tensors to memory (CPU) and initialize
-		TConstArrayView<FTensorDesc> InputSymbolicTensorDescs = InferenceModel->GetInputTensorDescs();
-		TArray<FTensorShape> InputShapes;
-		TArray<FTensor> InputTensorsDescs;
+		// If test does not ask for specific input/output, fill with model default converting variable dim to 1 if any.
+		TArray<FTensor> InputTensors;
+		
+		FillTensors(TestSetup.Inputs, InferenceModel->GetInputTensorDescs(), InputTensors);
+		FillTensors(TestSetup.Outputs, InferenceModel->GetOutputTensorDescs(), OutOutputTensors);
+		
+		// bind tensors to memory (CPU) and initialize
 		TArray<FMLTensorBinding> InputBindings;
+		TArray<FMLTensorBinding> OutputBindings;
 		TArray<TArray<char>> InputMemBuffers;
+		TArray<TArray<char>> OutputMemBuffers;
 		
-		for (const FTensorDesc& SymbolicTensorDesc : InputSymbolicTensorDescs)
-		{
-			FTensor TensorDesc = FTensor::MakeFromSymbolicDesc(SymbolicTensorDesc);
-			InputShapes.Emplace(TensorDesc.GetShape());
-			InputTensorsDescs.Emplace(TensorDesc);
-		}
-		FillTensorBindings(InputTensorsDescs, InputMemBuffers, InputBindings, InputTensorInitializer);
-		
-		//bind output tensors to memory (CPU) and initialize
-		TConstArrayView<FTensorDesc> OutputSymbolicTensorDescs = InferenceModel->GetOutputTensorDescs();
-		TArray<FTensor> OutputTensorsDescs;
-		TArray<NNX::FMLTensorBinding> OutputBindings;
-
-		for (const FTensorDesc& SymbolicTensorDesc : OutputSymbolicTensorDescs)
-		{
-			FTensor TensorDesc = FTensor::MakeFromSymbolicDesc(SymbolicTensorDesc);
-			OutputTensorsDescs.Emplace(TensorDesc);
-		}
-		OutOutputTensorsDesc = OutputTensorsDescs;
-		FillTensorBindings(OutputTensorsDescs, OutOutputMemBuffers, OutputBindings, OutputTensorInitializer);
+		FillTensorBindings(InputTensors, InputMemBuffers, InputBindings, InputTensorInitializer);
+		FillTensorBindings(OutOutputTensors, OutOutputMemBuffers, OutputBindings, OutputTensorInitializer);
 
-		check(OutOutputTensorsDesc.Num() == OutOutputMemBuffers.Num());
+		check(InputTensors.Num() == InputMemBuffers.Num());
+		check(OutOutputTensors.Num() == OutOutputMemBuffers.Num());
 
-		//To help for debugging sessions.
+		// To help for debugging sessions.
 		#ifdef UE_BUILD_DEBUG
 		const constexpr int32 NumTensorPtrForDebug = 3;
 		float* InputsAsFloat[NumTensorPtrForDebug];
@@ -397,8 +400,29 @@ namespace Test
 		}
 		#endif //UE_BUILD_DEBUG
 
-		//run inference
-		int ReturnValue = InferenceModel->Run(InputBindings, InputShapes, OutputBindings);
+		
+		TArray<FTensorShape> InputShapes;
+		for (const FTensor& Tensor : InputTensors)
+		{
+			InputShapes.Emplace(Tensor.GetShape());
+		}
+
+		// Setup inputs
+		int ReturnValue = InferenceModel->SetInputTensorShapes(InputShapes);
+		if (ReturnValue != 0)
+		{
+			UE_LOG(LogNNX, Error, TEXT("Failed to set input tensor shapes."));
+			return ReturnValue;
+		}
+
+		// Run inference
+		ReturnValue = InferenceModel->Run(InputBindings, OutputBindings);
+		if (ReturnValue != 0)
+		{
+			UE_LOG(LogNNX, Error, TEXT("Failed to run the model."));
+			return ReturnValue;
+		}
+		
 		return ReturnValue;
 	}
 
@@ -413,7 +437,7 @@ namespace Test
 		UE_LOG(LogNNX, Display, TEXT("Starting tests of '%s'"), *TestSetup.TargetName);
 
 		// Reference runtime
-		NNX::IRuntime* RefRuntime = NNX::GetRuntime(TEXT("NNXRuntimeCPU"));
+		IRuntime* RefRuntime = NNX::GetRuntime(TEXT("NNXRuntimeCPU"));
 		if (!RefRuntime)
 		{
 			UE_LOG(LogNNX, Error, TEXT("Can't load NNXRuntimeCPU runtime. Tests ABORTED!"));
@@ -421,15 +445,15 @@ namespace Test
 		}
 		const FString& RefName = RefRuntime->GetRuntimeName();
 		TArray<TArray<char>> RefOutputMemBuffers;
-		TArray<NNX::FTensor> RefOutputTensorDescs;
+		TArray<FTensor> RefOutputTensorDescs;
 
-		RunTestInference(ONNXModel, RefRuntime, RefOutputTensorDescs, RefOutputMemBuffers);
+		RunTestInference(ONNXModel, TestSetup, RefRuntime, RefOutputTensorDescs, RefOutputMemBuffers);
 
 		// Test against other runtime
-		TArray<NNX::IRuntime*> Runtimes;
+		TArray<IRuntime*> Runtimes;
 		bool bAllTestsSucceeded = true;
 
-		Runtimes = NNX::GetAllRuntimes();
+		Runtimes = GetAllRuntimes();
 		for (auto Runtime : Runtimes)
 		{
 			const FString& RuntimeName = Runtime->GetRuntimeName();
@@ -459,12 +483,12 @@ namespace Test
 			else
 			{
 				TArray<TArray<char>> OutputMemBuffers;
-				TArray<NNX::FTensor> OutputTensorDescs;
+				TArray<FTensor> OutputTensorDescs;
 				bool bTestSuceeded = true;
 				float AbsoluteErrorEpsilon = TestSetup.GetAbsoluteErrorEpsilonForRuntime(RuntimeName);
 				float RelativeErrorPercent = TestSetup.GetRelativeErrorPercentForRuntime(RuntimeName);
 
-				RunTestInference(ONNXModel, Runtime, OutputTensorDescs, OutputMemBuffers);
+				RunTestInference(ONNXModel, TestSetup, Runtime, OutputTensorDescs, OutputMemBuffers);
 				if (OutputTensorDescs.Num() == RefOutputTensorDescs.Num())
 				{
 					for (int i = 0; i < OutputTensorDescs.Num(); ++i)
diff --git a/Engine/Plugins/Experimental/NNX/Source/NNXQA/Resources/NNXQATestDesc.json b/Engine/Plugins/Experimental/NNX/Source/NNXQA/Resources/NNXQATestDesc.json
--- a/Engine/Plugins/Experimental/NNX/Source/NNXQA/Resources/NNXQATestDesc.json
+++ b/Engine/Plugins/Experimental/NNX/Source/NNXQA/Resources/NNXQATestDesc.json
@@ -79,6 +79,28 @@
 					"tags": [ "SmokeTest" ]
 				}
 			]
+		},
+		{
+			"category": "VariableInputShapes",
+			"is_model_test": true,
+			"targets": [
+				{
+					"target": "OneLayerNoWeightVariableInputShape",
+					"tags": [ "SmokeTest" ],
+					"runtimes": [
+						{ "name": "NNXRuntimeHLSL", "skip": true, "comment": "RDG does not support variable input sized tensor" },
+						{ "name": "NNXRuntimeDML",  "skip": true, "comment": "RDG does not support variable input sized tensor" }
+					]
+				},
+				{
+					"target": "TwoLayersNoWeightVariableInputShape",
+					"tags": [ "SmokeTest" ],
+					"runtimes": [
+						{ "name": "NNXRuntimeHLSL", "skip": true, "comment": "RDG does not support variable input sized tensor" },
+						{ "name": "NNXRuntimeDML",  "skip": true, "comment": "RDG does not support variable input sized tensor" }
+					]
+				}
+			]
 		}
 	],
 
@@ -258,10 +280,17 @@
 	],
 
 	"input_output_sets": [
+		{
+			"name": "VariableInputShapes.Regular",
+			"datasets": [
+				{ "inputs": [ { "shape": [ 1,1 ] } ], "comment": "if not specified output is defaulted as the first input" },
+				{ "inputs": [ { "shape": [ 2,3 ] } ] }
+			]
+		},
 		{
 			"name": "Elementwise.Unary.Regular",
 			"datasets": [
-				{ "inputs": [ { "shape": [ 1 ] } ], "comment": "if not specified output is defaulted as the first input" },
+				{ "inputs": [ { "shape": [ 1 ] } ] },
 				{ "inputs": [ { "shape": [ 1,2,3,4 ] } ] },
 				{ "inputs": [ { "shape": [ 1,512 ] } ] }
 			]
diff --git a/Engine/Plugins/Experimental/NNX/Source/NNXRuntimeCPU/Private/NNXRuntimeCPU.cpp b/Engine/Plugins/Experimental/NNX/Source/NNXRuntimeCPU/Private/NNXRuntimeCPU.cpp
--- a/Engine/Plugins/Experimental/NNX/Source/NNXRuntimeCPU/Private/NNXRuntimeCPU.cpp
+++ b/Engine/Plugins/Experimental/NNX/Source/NNXRuntimeCPU/Private/NNXRuntimeCPU.cpp
@@ -191,13 +191,7 @@ bool FMLInferenceModelCPU::ConfigureTensors(const bool InIsInput)
 		}
 
 		FTensorDesc SymbolicTensorDesc = FTensorDesc::Make(FString(CurTensorName), Shape, TypeAndSize.first);
-		if (!SymbolicTensorDesc.IsConcrete())
-			{
-				UE_LOG(LogNNX, Display,
-					TEXT("Negative (i.e., variable) dimensions not allowed yet, hard-coded to 1. Let us know if you really need variable dimensions."
-						" Keep in mind that fixed sizes might allow additional optimizations and speedup of the network during Run()."));
-			}
-
+		
 		check(SymbolicTensorDesc.GetElemByteSize() == TypeAndSize.second);
 		SymbolicTensorDescs.Emplace(SymbolicTensorDesc);
 	}
@@ -205,33 +199,46 @@ bool FMLInferenceModelCPU::ConfigureTensors(const bool InIsInput)
 	return true;
 }
 
-int FMLInferenceModelCPU::SetInputShapes(TConstArrayView<FTensorShape> InInputShapes)
+int FMLInferenceModelCPU::SetInputTensorShapes(TConstArrayView<FTensorShape> InInputShapes)
 {
-	//Verify input shape are valid for the model
-	if (FMLInferenceModel::SetInputShapes(InInputShapes) != 0)
+	// Verify input shape are valid for the model and set InputTensorShapes
+	if (FMLInferenceModel::SetInputTensorShapes(InInputShapes) != 0)
 	{
 		return -1;
 	}
 
-	//Run shape inference
-	//TODO jira 168972: handle dynamic input shape
-
-	//Build concrete input and output tensor descriptions
+	// Setup concrete input tensor
 	InputTensors.Empty();
-	for (FTensorDesc SymbolicTensorDesc : InputSymbolicTensors)
+	for (int i = 0; i < InputSymbolicTensors.Num(); ++i)
 	{
-		InputTensors.Emplace(FTensor::MakeFromSymbolicDesc(SymbolicTensorDesc));
+		FTensor Tensor = FTensor::Make(InputSymbolicTensors[i].GetName(), InInputShapes[i], InputSymbolicTensors[i].GetDataType());
+		InputTensors.Emplace(Tensor);
 	}
+
+	// Here model optimization could be done now that we know the input shapes, for some models
+	// this would allow to resolve output shapes here rather than during inference.
+
+	// Setup concrete output shapes only if all model output shapes are concretes, otherwise it will be set during Run()
 	OutputTensors.Empty();
 	for (FTensorDesc SymbolicTensorDesc : OutputSymbolicTensors)
 	{
-		OutputTensors.Emplace(FTensor::MakeFromSymbolicDesc(SymbolicTensorDesc));
+		if (SymbolicTensorDesc.IsConcrete())
+		{
+			FTensor Tensor = FTensor::MakeFromSymbolicDesc(SymbolicTensorDesc);
+			OutputTensors.Emplace(Tensor);
+			OutputTensorShapes.Emplace(Tensor.GetShape());
+		}
+	}
+	if (OutputTensors.Num() != OutputSymbolicTensors.Num())
+	{
+		OutputTensors.Empty();
+		OutputTensorShapes.Empty();
 	}
 
 	return 0;
 }
 
-int FMLInferenceModelCPU::Run(TConstArrayView<FMLTensorBinding> InInputBindings, TConstArrayView<FTensorShape> InInputShapes, TConstArrayView<FMLTensorBinding> InOutputBindings)
+int FMLInferenceModelCPU::Run(TConstArrayView<FMLTensorBinding> InInputBindings, TConstArrayView<FMLTensorBinding> InOutputBindings)
 {
 	DECLARE_SCOPE_CYCLE_COUNTER(TEXT("FMLInferenceModelCPU_Run"), STAT_FMLInferenceModelCPU_Run, STATGROUP_MachineLearning);
 
@@ -242,11 +249,10 @@ int FMLInferenceModelCPU::Run(TConstArrayView<FMLTensorBinding> InInputBindings,
 		return -1;
 	}
 
-	// Prepare the model for the concrete input shapes
-	int Res = SetInputShapes(InInputShapes);
-	if (Res != 0)
+	// Verify the model inputs were prepared
+	if (InputTensorShapes.Num() == 0)
 	{
-		UE_LOG(LogNNX, Warning, TEXT("Model preparation failed with error code:%d"), Res);
+		UE_LOG(LogNNX, Error, TEXT("Run(): Input shapes are not set, please call SetInputTensorShapes."));
 		return -1;
 	}
 
@@ -265,13 +271,36 @@ int FMLInferenceModelCPU::Run(TConstArrayView<FMLTensorBinding> InInputBindings,
 		TArray<Ort::Value> InputOrtTensors;
 		BindTensorsToORT(InInputBindings, InputTensors, InputTensorsORTType, AllocatorInfo.Get(), InputOrtTensors);
 
-		TArray<Ort::Value> OutputOrtTensors;
-		BindTensorsToORT(InOutputBindings, OutputTensors, OutputTensorsORTType, AllocatorInfo.Get(), OutputOrtTensors);
+		if (!OutputTensors.IsEmpty())
+		{
+			// If output shapes are known we can directly map preallocated output buffers
+			TArray<Ort::Value> OutputOrtTensors;
+			BindTensorsToORT(InOutputBindings, OutputTensors, OutputTensorsORTType, AllocatorInfo.Get(), OutputOrtTensors);
 
-		Session->Run(Ort::RunOptions{ nullptr },
-			InputTensorNames.GetData(), &InputOrtTensors[0], InputTensorNames.Num(),
-			OutputTensorNames.GetData(), &OutputOrtTensors[0], OutputTensorNames.Num());
+			Session->Run(Ort::RunOptions{ nullptr },
+				InputTensorNames.GetData(), &InputOrtTensors[0], InputTensorNames.Num(),
+				OutputTensorNames.GetData(), &OutputOrtTensors[0], OutputTensorNames.Num());
+		}
+		else
+		{
+			TArray<Ort::Value> OutputOrtTensors;
+			for (int i = 0; i < InOutputBindings.Num(); ++i)
+			{
+				OutputOrtTensors.Emplace(nullptr);
+			}
 
+			Session->Run(Ort::RunOptions{ nullptr },
+				InputTensorNames.GetData(), &InputOrtTensors[0], InputTensorNames.Num(),
+				OutputTensorNames.GetData(), &OutputOrtTensors[0], OutputTensorNames.Num());
+			
+			// Output shapes were resolved during inference: Copy the data back to bindings and expose output tensor shapes
+			CopyFromORTToBindings(OutputOrtTensors, InOutputBindings, OutputSymbolicTensors, OutputTensors);
+			check(OutputTensorShapes.IsEmpty());
+			for (int i = 0; i < OutputTensors.Num(); ++i)
+			{
+				OutputTensorShapes.Emplace(OutputTensors[i].GetShape());
+			}
+		}
 	}
 #if WITH_EDITOR
 	catch (const std::exception& Exception)
diff --git a/Engine/Plugins/Experimental/NNX/Source/NNXRuntimeCPU/Public/NNXRuntimeCPU.h b/Engine/Plugins/Experimental/NNX/Source/NNXRuntimeCPU/Public/NNXRuntimeCPU.h
--- a/Engine/Plugins/Experimental/NNX/Source/NNXRuntimeCPU/Public/NNXRuntimeCPU.h
+++ b/Engine/Plugins/Experimental/NNX/Source/NNXRuntimeCPU/Public/NNXRuntimeCPU.h
@@ -84,8 +84,8 @@ namespace NNX
 		bool Init(const UMLInferenceModel* InferenceModel);
 		bool IsLoaded() const;
 
-		virtual int SetInputShapes(TConstArrayView<FTensorShape> InInputShapes) override;
-		virtual int Run(TConstArrayView<FMLTensorBinding> InInputBindings, TConstArrayView<FTensorShape> InInputShapes, TConstArrayView<FMLTensorBinding> InOutputBindings) override;
+		virtual int SetInputTensorShapes(TConstArrayView<FTensorShape> InInputShapes) override;
+		virtual int Run(TConstArrayView<FMLTensorBinding> InInputBindings, TConstArrayView<FMLTensorBinding> InOutputBindings) override;
 
 		float GetLastRunTimeMSec() const;
 		UE::NNEProfiling::Internal::FStatistics GetRunStatistics() const;
diff --git a/Engine/Plugins/Experimental/NNX/Source/NNXRuntimeCPU/Public/NNXRuntimeCPUUtils.h b/Engine/Plugins/Experimental/NNX/Source/NNXRuntimeCPU/Public/NNXRuntimeCPUUtils.h
--- a/Engine/Plugins/Experimental/NNX/Source/NNXRuntimeCPU/Public/NNXRuntimeCPUUtils.h
+++ b/Engine/Plugins/Experimental/NNX/Source/NNXRuntimeCPU/Public/NNXRuntimeCPUUtils.h
@@ -2,6 +2,7 @@
 
 #pragma once
 
+#include "HAL/UnrealMemory.h"
 #include "NNXTypes.h"
 #include "NNXRuntime.h"
 #include "NNXThirdPartyWarningDisabler.h"
@@ -112,15 +113,15 @@ inline TypeInfoORT TranslateTensorTypeORTToNNI(unsigned int OrtDataType)
 }
 
 inline void BindTensorsToORT(
-	TArrayView<const NNX::FMLTensorBinding> InBindingTensors,
-	const TArray<NNX::FTensor>& InTensorsDescriptors,
-	const TArray<ONNXTensorElementDataType>& InTensorsORTType, 
+	TConstArrayView<const NNX::FMLTensorBinding> InBindingTensors,
+	TConstArrayView<NNX::FTensor> InTensors,
+	TConstArrayView<ONNXTensorElementDataType> InTensorsORTType, 
 	const Ort::MemoryInfo* InAllocatorInfo,
-	TArray<Ort::Value>& OutOrtTensors	
+	TArray<Ort::Value>& OutOrtTensors
 )
 {
 	const uint32 NumBinding = InBindingTensors.Num();
-	const uint32 NumDescriptors = InTensorsDescriptors.Num();
+	const uint32 NumDescriptors = InTensors.Num();
 
 	if (NumBinding != NumDescriptors)
 	{
@@ -130,23 +131,23 @@ inline void BindTensorsToORT(
 
 	for (uint32 Index = 0; Index < NumBinding; ++Index)
 	{
-		const NNX::FMLTensorBinding& CurrentBinding = InBindingTensors[Index];
-		const NNX::FTensor& CurrentDescriptor = InTensorsDescriptors[Index];
-		ONNXTensorElementDataType CurrentORTType = InTensorsORTType[Index];
+		const NNX::FMLTensorBinding& Binding = InBindingTensors[Index];
+		const NNX::FTensor& Tensor = InTensors[Index];
+		const ONNXTensorElementDataType CurrentORTType = InTensorsORTType[Index];
 
 		TUniquePtr<int64_t[]> SizesInt64t;
-		SizesInt64t = MakeUnique<int64_t[]>(CurrentDescriptor.GetShape().Rank());
-		for (int32 DimIndex = 0; DimIndex < CurrentDescriptor.GetShape().Rank(); ++DimIndex)
+		SizesInt64t = MakeUnique<int64_t[]>(Tensor.GetShape().Rank());
+		for (int32 DimIndex = 0; DimIndex < Tensor.GetShape().Rank(); ++DimIndex)
 		{
-			SizesInt64t.Get()[DimIndex] = CurrentDescriptor.GetShape().Data[DimIndex];
+			SizesInt64t.Get()[DimIndex] = Tensor.GetShape().Data[DimIndex];
 		}
 
-		const uint64 ByteCount { InTensorsDescriptors[Index].GetDataSize() };
-		const uint32 ArrayDimensions { (uint32)CurrentDescriptor.GetShape().Rank() };
+		const uint64 ByteCount { InTensors[Index].GetDataSize() };
+		const uint32 ArrayDimensions { (uint32)Tensor.GetShape().Rank() };
 		OutOrtTensors.Emplace(
 			Ort::Value::CreateTensor(
 				*InAllocatorInfo,
-				CurrentBinding.CpuMemory,
+				Binding.CpuMemory,
 				ByteCount, 
 				SizesInt64t.Get(),
 				ArrayDimensions, 
@@ -154,4 +155,53 @@ inline void BindTensorsToORT(
 			)
 		);
 	}
+}
+
+inline void CopyFromORTToBindings(
+	TConstArrayView<Ort::Value> InOrtTensors,
+	TConstArrayView<const NNX::FMLTensorBinding> InBindingTensors,
+	TConstArrayView<NNX::FTensorDesc> InTensorDescs,
+	TArray<NNX::FTensor>& OutTensors
+)
+{
+	const uint32 NumBinding = InOrtTensors.Num();
+	const uint32 NumDescriptors = InTensorDescs.Num();
+
+	if (NumBinding != NumDescriptors)
+	{
+		UE_LOG(LogNNX, Error, TEXT("CopyFromORTToBindings: Number of Bindings is different of Descriptors."));
+		return;
+	}
+
+	for (uint32 Index = 0; Index < NumBinding; ++Index)
+	{
+		const NNX::FMLTensorBinding& Binding = InBindingTensors[Index];
+		const NNX::FTensorDesc& TensorDesc = InTensorDescs[Index];
+		const Ort::Value& OrtTensor = InOrtTensors[Index];
+		const std::vector<int64_t>& OrtShape = OrtTensor.GetTensorTypeAndShapeInfo().GetShape();
+
+		NNX::FTensorShape Shape;
+		for (int32 DimIndex = 0; DimIndex < OrtShape.size(); ++DimIndex)
+		{
+			check(OrtShape[DimIndex] >= 0);
+			Shape.Data.Add(OrtShape[DimIndex]);
+		}
+
+		NNX::FTensor Tensor = NNX::FTensor::Make(TensorDesc.GetName(), Shape, TensorDesc.GetDataType());
+		OutTensors.Emplace(Tensor);
+
+		void* CpuMemory = (uint8_t*)Binding.CpuMemory + Binding.OffsetInBytes;
+		const void* OrtTensorData = OrtTensor.GetTensorData<void>();
+		const uint64 DataSize = FMath::Min(Tensor.GetDataSize(), Binding.SizeInBytes);
+
+		if (DataSize >= 0)
+		{
+			FMemory::Memcpy(CpuMemory, OrtTensorData, DataSize);
+		}
+
+		if (Tensor.GetDataSize() > Binding.SizeInBytes)
+		{
+			UE_LOG(LogNNX, Error, TEXT("CopyFromORTToBindings: Binding buffer was not large enough to contain all of the data, only first %d bytes were copied."), DataSize);
+		}
+	}
 }
\ No newline at end of file
diff --git a/Engine/Plugins/Experimental/NNX/Source/NNXRuntimeORT/Private/NNXRuntimeORT.cpp b/Engine/Plugins/Experimental/NNX/Source/NNXRuntimeORT/Private/NNXRuntimeORT.cpp
--- a/Engine/Plugins/Experimental/NNX/Source/NNXRuntimeORT/Private/NNXRuntimeORT.cpp
+++ b/Engine/Plugins/Experimental/NNX/Source/NNXRuntimeORT/Private/NNXRuntimeORT.cpp
@@ -265,13 +265,6 @@ bool FMLInferenceModelORT::ConfigureTensors(const bool InIsInput)
 		}
 
 		FTensorDesc SymbolicTensorDesc = FTensorDesc::Make(FString(CurTensorName), Shape, TypeAndSize.first);
-
-		if (!SymbolicTensorDesc.IsConcrete())
-			{
-				UE_LOG(LogNNX, Display,
-					TEXT("Negative (i.e., variable) dimensions not allowed yet, hard-coded to 1. Let us know if you really need variable dimensions."
-						" Keep in mind that fixed sizes might allow additional optimizations and speedup of the network during Run()."));
-			}
 		
 		check(SymbolicTensorDesc.GetElemByteSize() == TypeAndSize.second);
 		SymbolicTensorDescs.Emplace(SymbolicTensorDesc);
@@ -280,33 +273,46 @@ bool FMLInferenceModelORT::ConfigureTensors(const bool InIsInput)
 	return true;
 }
 
-int FMLInferenceModelORT::SetInputShapes(TConstArrayView<FTensorShape> InInputShapes)
+int FMLInferenceModelORT::SetInputTensorShapes(TConstArrayView<FTensorShape> InInputShapes)
 {
-	//Verify input shape are valid for the model
-	if (FMLInferenceModel::SetInputShapes(InInputShapes) != 0)
+	// Verify input shape are valid for the model and set InputTensorShapes
+	if (FMLInferenceModel::SetInputTensorShapes(InInputShapes) != 0)
 	{
 		return -1;
 	}
 
-	//Run shape inference
-	//TODO jira 168972: handle dynamic input shape
-
-	//Build concrete input and output tensor descriptions
+	// Setup concrete input tensor
 	InputTensors.Empty();
-	for (FTensorDesc SymbolicTensorDesc : InputSymbolicTensors)
+	for (int i = 0; i < InputSymbolicTensors.Num(); ++i)
 	{
-		InputTensors.Emplace(FTensor::MakeFromSymbolicDesc(SymbolicTensorDesc));
+		FTensor Tensor = FTensor::Make(InputSymbolicTensors[i].GetName(), InInputShapes[i], InputSymbolicTensors[i].GetDataType());
+		InputTensors.Emplace(Tensor);
 	}
+
+	// Here model optimization could be done now that we know the input shapes, for some models
+	// this would allow to resolve output shapes here rather than during inference.
+
+	// Setup concrete output shapes only if all model output shapes are concretes, otherwise it will be set during Run()
 	OutputTensors.Empty();
 	for (FTensorDesc SymbolicTensorDesc : OutputSymbolicTensors)
 	{
-		OutputTensors.Emplace(FTensor::MakeFromSymbolicDesc(SymbolicTensorDesc));
+		if (SymbolicTensorDesc.IsConcrete())
+		{
+			FTensor Tensor = FTensor::MakeFromSymbolicDesc(SymbolicTensorDesc);
+			OutputTensors.Emplace(Tensor);
+			OutputTensorShapes.Emplace(Tensor.GetShape());
+		}
+	}
+	if (OutputTensors.Num() != OutputSymbolicTensors.Num())
+	{
+		OutputTensors.Empty();
+		OutputTensorShapes.Empty();
 	}
 
 	return 0;
 }
 
-int FMLInferenceModelORT::Run(TConstArrayView<FMLTensorBinding> InInputBindings, TConstArrayView<FTensorShape> InInputShapes, TConstArrayView<FMLTensorBinding> InOutputBindings)
+int FMLInferenceModelORT::Run(TConstArrayView<FMLTensorBinding> InInputBindings, TConstArrayView<FMLTensorBinding> InOutputBindings)
 {
 	DECLARE_SCOPE_CYCLE_COUNTER(TEXT("FMLInferenceModelORT_Run"), STAT_FMLInferenceModelORT_Run, STATGROUP_MachineLearning);
 
@@ -317,11 +323,10 @@ int FMLInferenceModelORT::Run(TConstArrayView<FMLTensorBinding> InInputBindings,
 		return -1;
 	}
 
-	// Prepare the model for the concrete input shapes
-	int Res = SetInputShapes(InInputShapes);
-	if (Res != 0)
+	// Verify the model inputs were prepared
+	if (InputTensorShapes.Num() == 0)
 	{
-		UE_LOG(LogNNX, Warning, TEXT("Model preparation failed with error code:%d"), Res);
+		UE_LOG(LogNNX, Error, TEXT("Run(): Input shapes are not set, please call SetInputTensorShapes."));
 		return -1;
 	}
 
@@ -340,13 +345,36 @@ int FMLInferenceModelORT::Run(TConstArrayView<FMLTensorBinding> InInputBindings,
 		TArray<Ort::Value> InputOrtTensors;
 		BindTensorsToORT(InInputBindings, InputTensors, InputTensorsORTType, AllocatorInfo.Get(), InputOrtTensors);
 
-		TArray<Ort::Value> OutputOrtTensors;
-		BindTensorsToORT(InOutputBindings, OutputTensors, OutputTensorsORTType, AllocatorInfo.Get(), OutputOrtTensors);
+		if (!OutputTensors.IsEmpty())
+		{
+			// If output shapes are known we can directly map preallocated output buffers
+			TArray<Ort::Value> OutputOrtTensors;
+			BindTensorsToORT(InOutputBindings, OutputTensors, OutputTensorsORTType, AllocatorInfo.Get(), OutputOrtTensors);
 
-		Session->Run(Ort::RunOptions{ nullptr },
-			InputTensorNames.GetData(), &InputOrtTensors[0], InputTensorNames.Num(),
-			OutputTensorNames.GetData(), &OutputOrtTensors[0], OutputTensorNames.Num());
+			Session->Run(Ort::RunOptions{ nullptr },
+				InputTensorNames.GetData(), &InputOrtTensors[0], InputTensorNames.Num(),
+				OutputTensorNames.GetData(), &OutputOrtTensors[0], OutputTensorNames.Num());
+		}
+		else
+		{
+			TArray<Ort::Value> OutputOrtTensors;
+			for (int i = 0; i < InOutputBindings.Num(); ++i)
+			{
+				OutputOrtTensors.Emplace(nullptr);
+			}
 
+			Session->Run(Ort::RunOptions{ nullptr },
+				InputTensorNames.GetData(), &InputOrtTensors[0], InputTensorNames.Num(),
+				OutputTensorNames.GetData(), &OutputOrtTensors[0], OutputTensorNames.Num());
+
+			// Output shapes were resolved during inference: Copy the data back to bindings and expose output tensor shapes
+			CopyFromORTToBindings(OutputOrtTensors, InOutputBindings, OutputSymbolicTensors, OutputTensors);
+			check(OutputTensorShapes.IsEmpty());
+			for (int i = 0; i < OutputTensors.Num(); ++i)
+			{
+				OutputTensorShapes.Emplace(OutputTensors[i].GetShape());
+			}
+		}
 	}
 #if WITH_EDITOR
 	catch (const std::exception& Exception)
@@ -356,6 +384,7 @@ int FMLInferenceModelORT::Run(TConstArrayView<FMLTensorBinding> InInputBindings,
 #endif //WITH_EDITOR
 
 	RunStatisticsEstimator.StoreSample(RunTimer.Toc());
+
 	return 0;
 }
 
diff --git a/Engine/Plugins/Experimental/NNX/Source/NNXRuntimeORT/Private/NNXRuntimeORT.h b/Engine/Plugins/Experimental/NNX/Source/NNXRuntimeORT/Private/NNXRuntimeORT.h
--- a/Engine/Plugins/Experimental/NNX/Source/NNXRuntimeORT/Private/NNXRuntimeORT.h
+++ b/Engine/Plugins/Experimental/NNX/Source/NNXRuntimeORT/Private/NNXRuntimeORT.h
@@ -183,8 +183,8 @@ namespace NNX
 		bool Init(const UMLInferenceModel* InferenceModel);
 		bool IsLoaded() const;
 
-		virtual int SetInputShapes(TConstArrayView<FTensorShape> InInputShapes) override;
-		virtual int Run(TConstArrayView<FMLTensorBinding> InInputBindings, TConstArrayView<FTensorShape> InInputShapes, TConstArrayView<FMLTensorBinding> InOutputBindings) override;
+		virtual int SetInputTensorShapes(TConstArrayView<FTensorShape> InInputShapes) override;
+		virtual int Run(TConstArrayView<FMLTensorBinding> InInputBindings, TConstArrayView<FMLTensorBinding> InOutputBindings) override;
 
 		float GetLastRunTimeMSec() const;
 		UE::NNEProfiling::Internal::FStatistics GetRunStatistics() const;
diff --git a/Engine/Plugins/Experimental/NNX/Source/NNXRuntimeORT/Private/NNXRuntimeORTUtils.h b/Engine/Plugins/Experimental/NNX/Source/NNXRuntimeORT/Private/NNXRuntimeORTUtils.h
--- a/Engine/Plugins/Experimental/NNX/Source/NNXRuntimeORT/Private/NNXRuntimeORTUtils.h
+++ b/Engine/Plugins/Experimental/NNX/Source/NNXRuntimeORT/Private/NNXRuntimeORTUtils.h
@@ -2,6 +2,7 @@
 
 #pragma once
 
+#include "HAL/UnrealMemory.h"
 #include "NNXTypes.h"
 #include "NNXRuntime.h"
 #include "NNXThirdPartyWarningDisabler.h"
@@ -15,7 +16,7 @@ DECLARE_STATS_GROUP(TEXT("MachineLearning"), STATGROUP_MachineLearning, STATCAT_
 
 using TypeInfoORT = std::pair<EMLTensorDataType, uint64>;
 
-inline TypeInfoORT TranslateTensorTypeORTToNNI(unsigned int OrtDataType) 
+inline TypeInfoORT TranslateTensorTypeORTToNNI(unsigned int OrtDataType)
 {
 	EMLTensorDataType DataType = EMLTensorDataType::None;
 	uint64 ElementSize = 0;
@@ -108,19 +109,19 @@ inline TypeInfoORT TranslateTensorTypeORTToNNI(unsigned int OrtDataType)
 		break;
 	}
 
-	return TypeInfoORT{DataType, ElementSize};
+	return TypeInfoORT{ DataType, ElementSize };
 }
 
 inline void BindTensorsToORT(
-	TArrayView<const NNX::FMLTensorBinding> InBindingTensors,
-	const TArray<NNX::FTensor>& InTensorsDescriptors,
-	const TArray<ONNXTensorElementDataType>& InTensorsORTType, 
+	TConstArrayView<const NNX::FMLTensorBinding> InBindingTensors,
+	TConstArrayView<NNX::FTensor> InTensors,
+	TConstArrayView<ONNXTensorElementDataType> InTensorsORTType,
 	const Ort::MemoryInfo* InAllocatorInfo,
-	TArray<Ort::Value>& OutOrtTensors	
+	TArray<Ort::Value>& OutOrtTensors
 )
 {
 	const uint32 NumBinding = InBindingTensors.Num();
-	const uint32 NumDescriptors = InTensorsDescriptors.Num();
+	const uint32 NumDescriptors = InTensors.Num();
 
 	if (NumBinding != NumDescriptors)
 	{
@@ -130,28 +131,77 @@ inline void BindTensorsToORT(
 
 	for (uint32 Index = 0; Index < NumBinding; ++Index)
 	{
-		const NNX::FMLTensorBinding& CurrentBinding = InBindingTensors[Index];
-		const NNX::FTensor& CurrentDescriptor = InTensorsDescriptors[Index];
-		ONNXTensorElementDataType CurrentORTType = InTensorsORTType[Index];
+		const NNX::FMLTensorBinding& Binding = InBindingTensors[Index];
+		const NNX::FTensor& Tensor = InTensors[Index];
+		const ONNXTensorElementDataType CurrentORTType = InTensorsORTType[Index];
 
 		TUniquePtr<int64_t[]> SizesInt64t;
-		SizesInt64t = MakeUnique<int64_t[]>(CurrentDescriptor.GetShape().Rank());
-		for (int32 DimIndex = 0; DimIndex < CurrentDescriptor.GetShape().Rank(); ++DimIndex)
+		SizesInt64t = MakeUnique<int64_t[]>(Tensor.GetShape().Rank());
+		for (int32 DimIndex = 0; DimIndex < Tensor.GetShape().Rank(); ++DimIndex)
 		{
-			SizesInt64t.Get()[DimIndex] = CurrentDescriptor.GetShape().Data[DimIndex];
+			SizesInt64t.Get()[DimIndex] = Tensor.GetShape().Data[DimIndex];
 		}
 
-		const uint64 ByteCount { InTensorsDescriptors[Index].GetDataSize() };
-		const uint32 ArrayDimensions { (uint32)CurrentDescriptor.GetShape().Rank() };
+		const uint64 ByteCount{ InTensors[Index].GetDataSize() };
+		const uint32 ArrayDimensions{ (uint32)Tensor.GetShape().Rank() };
 		OutOrtTensors.Emplace(
 			Ort::Value::CreateTensor(
 				*InAllocatorInfo,
-				CurrentBinding.CpuMemory,
-				ByteCount, 
+				Binding.CpuMemory,
+				ByteCount,
 				SizesInt64t.Get(),
-				ArrayDimensions, 
+				ArrayDimensions,
 				CurrentORTType
 			)
 		);
 	}
+}
+
+inline void CopyFromORTToBindings(
+	TConstArrayView<Ort::Value> InOrtTensors,
+	TConstArrayView<const NNX::FMLTensorBinding> InBindingTensors,
+	TConstArrayView<NNX::FTensorDesc> InTensorDescs,
+	TArray<NNX::FTensor>& OutTensors
+)
+{
+	const uint32 NumBinding = InOrtTensors.Num();
+	const uint32 NumDescriptors = InTensorDescs.Num();
+
+	if (NumBinding != NumDescriptors)
+	{
+		UE_LOG(LogNNX, Error, TEXT("CopyFromORTToBindings: Number of Bindings is different of Descriptors."));
+		return;
+	}
+
+	for (uint32 Index = 0; Index < NumBinding; ++Index)
+	{
+		const NNX::FMLTensorBinding& Binding = InBindingTensors[Index];
+		const NNX::FTensorDesc& TensorDesc = InTensorDescs[Index];
+		const Ort::Value& OrtTensor = InOrtTensors[Index];
+		const std::vector<int64_t>& OrtShape = OrtTensor.GetTensorTypeAndShapeInfo().GetShape();
+
+		NNX::FTensorShape Shape;
+		for (int32 DimIndex = 0; DimIndex < OrtShape.size(); ++DimIndex)
+		{
+			check(OrtShape[DimIndex] >= 0);
+			Shape.Data.Add(OrtShape[DimIndex]);
+		}
+
+		NNX::FTensor Tensor = NNX::FTensor::Make(TensorDesc.GetName(), Shape, TensorDesc.GetDataType());
+		OutTensors.Emplace(Tensor);
+
+		void* CpuMemory = (uint8_t*)Binding.CpuMemory + Binding.OffsetInBytes;
+		const void* OrtTensorData = OrtTensor.GetTensorData<void>();
+		const uint64 DataSize = FMath::Min(Tensor.GetDataSize(), Binding.SizeInBytes);
+
+		if (DataSize >= 0)
+		{
+			FMemory::Memcpy(CpuMemory, OrtTensorData, DataSize);
+		}
+
+		if (Tensor.GetDataSize() > Binding.SizeInBytes)
+		{
+			UE_LOG(LogNNX, Error, TEXT("CopyFromORTToBindings: Binding buffer was not large enough to contain all of the data, only first %d bytes were copied."), DataSize);
+		}
+	}
 }
\ No newline at end of file
diff --git a/Engine/Plugins/Experimental/NNX/Source/NNXRuntimeRDG/Private/NNXRuntimeDML.cpp b/Engine/Plugins/Experimental/NNX/Source/NNXRuntimeRDG/Private/NNXRuntimeDML.cpp
--- a/Engine/Plugins/Experimental/NNX/Source/NNXRuntimeRDG/Private/NNXRuntimeDML.cpp
+++ b/Engine/Plugins/Experimental/NNX/Source/NNXRuntimeRDG/Private/NNXRuntimeDML.cpp
@@ -1522,7 +1522,7 @@ void FMLInferenceModelDml::AddDispatchOps_RenderThread(FRDGBuilder& GraphBuilder
 	checkCode(
 		for (int32 i = 0; i < AllTensorBindings.Num(); ++i)
 		{
-			check(AllTensors[i].GetDataSize() == AllTensorBindings[i].SizeInBytes);
+			check(AllTensors[i].GetDataSize() <= AllTensorBindings[i].SizeInBytes);
 		}
 	);
 
diff --git a/Engine/Plugins/Experimental/NNX/Source/NNXRuntimeRDG/Private/NNXRuntimeHLSL.cpp b/Engine/Plugins/Experimental/NNX/Source/NNXRuntimeRDG/Private/NNXRuntimeHLSL.cpp
--- a/Engine/Plugins/Experimental/NNX/Source/NNXRuntimeRDG/Private/NNXRuntimeHLSL.cpp
+++ b/Engine/Plugins/Experimental/NNX/Source/NNXRuntimeRDG/Private/NNXRuntimeHLSL.cpp
@@ -143,7 +143,7 @@ void FMLInferenceModelHlsl::AddDispatchOps_RenderThread(FRDGBuilder& GraphBuilde
 	checkCode(
 		for (int32 i = 0; i < AllTensorBindings.Num(); ++i)
 		{
-			check(AllTensors[i].GetDataSize() == AllTensorBindings[i].SizeInBytes);
+			check(AllTensors[i].GetDataSize() <= AllTensorBindings[i].SizeInBytes);
 		}
 	);
 
diff --git a/Engine/Plugins/Experimental/NNX/Source/NNXRuntimeRDG/Private/NNXRuntimeRDG.cpp b/Engine/Plugins/Experimental/NNX/Source/NNXRuntimeRDG/Private/NNXRuntimeRDG.cpp
--- a/Engine/Plugins/Experimental/NNX/Source/NNXRuntimeRDG/Private/NNXRuntimeRDG.cpp
+++ b/Engine/Plugins/Experimental/NNX/Source/NNXRuntimeRDG/Private/NNXRuntimeRDG.cpp
@@ -244,21 +244,21 @@ bool FMLInferenceModelRDG::LoadModel(const FNNIModelRaw& InModel, FMLRuntimeForm
 /**
  * Run the inference model (synchronous version)
  */
-int FMLInferenceModelRDG::Run(TConstArrayView<FMLTensorBinding> InInputBindings, TConstArrayView<FTensorShape> InInputShapes, TConstArrayView<FMLTensorBinding> InOutputBindings)
+int FMLInferenceModelRDG::Run(TConstArrayView<FMLTensorBinding> InInputBindings, TConstArrayView<FMLTensorBinding> InOutputBindings)
 {
-	// Prepare the model for the concrete input shapes
-	int Res = SetInputShapes(InInputShapes);
-	if (Res != 0)
+	// Verify the model inputs were prepared
+	if (InputTensorShapes.Num() == 0)
 	{
-		UE_LOG(LogNNX, Warning, TEXT("Model preparation failed with error code:%d"), Res);
+		UE_LOG(LogNNX, Error, TEXT("Run(): Input shapes are not set, please call SetInputTensorShapes."));
 		return -1;
 	}
 	
+	int Res = 0;
 	FEvent* Signal = FGenericPlatformProcess::GetSynchEventFromPool(false);
 
 	ENQUEUE_RENDER_COMMAND(FMLInferenceModel_Run)
 	(
-		[&Signal, &Res, this, InInputBindings, InInputShapes, InOutputBindings](FRHICommandListImmediate& RHICmdList)
+		[&Signal, &Res, this, InInputBindings, InOutputBindings](FRHICommandListImmediate& RHICmdList)
 		{
 			TOptional<ERHIPipeline>		Pipeline = RHICmdList.GetPipeline();
 
@@ -269,7 +269,7 @@ int FMLInferenceModelRDG::Run(TConstArrayView<FMLTensorBinding> InInputBindings,
 
 			FRDGBuilder	RDGBuilder(RHICmdList);
 
-			Res = EnqueueRDG(RDGBuilder, InInputBindings, InInputShapes, InOutputBindings);
+			Res = EnqueueRDG(RDGBuilder, InInputBindings, InOutputBindings);
 			if (Res == 0)
 			{
 				RDGBuilder.Execute();
@@ -305,10 +305,10 @@ int FMLInferenceModelRDG::Run(TConstArrayView<FMLTensorBinding> InInputBindings,
 	return Res;
 }
 
-int FMLInferenceModelRDG::SetInputShapes(TConstArrayView<FTensorShape> InInputShapes)
+int FMLInferenceModelRDG::SetInputTensorShapes(TConstArrayView<FTensorShape> InInputShapes)
 {
-	//Verify input shape are valid for the model
-	if (FMLInferenceModel::SetInputShapes(InInputShapes) != 0)
+	//Verify input shape are valid for the model and set InputTensorShapes
+	if (FMLInferenceModel::SetInputTensorShapes(InInputShapes) != 0)
 	{
 		return -1;
 	}
@@ -329,10 +329,13 @@ int FMLInferenceModelRDG::SetInputShapes(TConstArrayView<FTensorShape> InInputSh
 	{
 		InputTensors.Emplace(AllTensors[InputIndices]);
 	}
+	
 	OutputTensors.Empty();
+	OutputTensorShapes.Empty();
 	for (int32 OutputIndices : OutputTensorIndices)
 	{
 		OutputTensors.Emplace(AllTensors[OutputIndices]);
+		OutputTensorShapes.Emplace(AllTensors[OutputIndices].GetShape());
 	}
 	check(InputTensorIndices.Num() + OutputTensorIndices.Num() + IntermediateTensorIndices.Num() == AllTensors.Num());
 	check(InputTensors.Num() + OutputTensors.Num() + IntermediateTensorIndices.Num() == AllSymbolicTensors.Num());
@@ -352,17 +355,16 @@ FRDGBufferDesc CreateRDGBufferDescFromTensorDesc(const FTensor& TensorDesc)
 /**
  * Enqueue operators to RDG, the caller will run the GraphBuilder.Execute()
  */
-int FMLInferenceModelRDG::EnqueueRDG(FRDGBuilder& RDGBuilder, TConstArrayView<FMLTensorBinding> InInputBindings, TConstArrayView<FTensorShape> InInputShapes, TConstArrayView<FMLTensorBinding> InOutputBindings)
+int FMLInferenceModelRDG::EnqueueRDG(FRDGBuilder& RDGBuilder, TConstArrayView<FMLTensorBinding> InInputBindings, TConstArrayView<FMLTensorBinding> InOutputBindings)
 {
 	check(IsInRenderingThread());
 
 	int Res;
 
-	// Prepare the model for the concrete input shapes
-	Res = SetInputShapes(InInputShapes);
-	if (Res != 0)
+	// Verify the model inputs were prepared
+	if (InputTensorShapes.Num() == 0)
 	{
-		UE_LOG(LogNNX, Warning, TEXT("Model preparation failed with error code:%d"), Res);
+		UE_LOG(LogNNX, Error, TEXT("EnqueueRDG(): Input shapes are not set, please call SetInputTensorShapes."));
 		return -1;
 	}
 
diff --git a/Engine/Plugins/Experimental/NNX/Source/NNXRuntimeRDG/Private/NNXRuntimeRDG.h b/Engine/Plugins/Experimental/NNX/Source/NNXRuntimeRDG/Private/NNXRuntimeRDG.h
--- a/Engine/Plugins/Experimental/NNX/Source/NNXRuntimeRDG/Private/NNXRuntimeRDG.h
+++ b/Engine/Plugins/Experimental/NNX/Source/NNXRuntimeRDG/Private/NNXRuntimeRDG.h
@@ -63,9 +63,9 @@ public:
 
 	~FMLInferenceModelRDG();
 
-	virtual int SetInputShapes(TConstArrayView<FTensorShape> InputShapes) override;
-	virtual int Run(TConstArrayView<FMLTensorBinding> InputBindings, TConstArrayView<FTensorShape> InputShapes, TConstArrayView<FMLTensorBinding> OutputBindings) override;
-	virtual int EnqueueRDG(FRDGBuilder& RDGBuilder, TConstArrayView<FMLTensorBinding> InputBindings, TConstArrayView<FTensorShape> InputShapes, TConstArrayView<FMLTensorBinding> OutputBindings) override;
+	virtual int SetInputTensorShapes(TConstArrayView<FTensorShape> InputShapes) override;
+	virtual int Run(TConstArrayView<FMLTensorBinding> InputBindings, TConstArrayView<FMLTensorBinding> OutputBindings) override;
+	virtual int EnqueueRDG(FRDGBuilder& RDGBuilder, TConstArrayView<FMLTensorBinding> InputBindings, TConstArrayView<FMLTensorBinding> OutputBindings) override;
 
 protected:
 
