NNI fix for Issue 235719: PVS Studio warnings in NNXRuntimeCPUUtils.h and NNXRuntimeORTUtils.h

diff --git a/Engine/Plugins/Experimental/NNX/Source/NNXRuntimeCPU/Public/NNXRuntimeCPUUtils.h b/Engine/Plugins/Experimental/NNX/Source/NNXRuntimeCPU/Public/NNXRuntimeCPUUtils.h
--- a/Engine/Plugins/Experimental/NNX/Source/NNXRuntimeCPU/Public/NNXRuntimeCPUUtils.h
+++ b/Engine/Plugins/Experimental/NNX/Source/NNXRuntimeCPU/Public/NNXRuntimeCPUUtils.h
@@ -194,7 +194,7 @@ inline void CopyFromORTToBindings(
 		const void* OrtTensorData = OrtTensor.GetTensorData<void>();
 		const uint64 DataSize = FMath::Min(Tensor.GetDataSize(), Binding.SizeInBytes);
 
-		if (DataSize >= 0)
+		if (DataSize > 0)
 		{
 			FMemory::Memcpy(CpuMemory, OrtTensorData, DataSize);
 		}
diff --git a/Engine/Plugins/Experimental/NNX/Source/NNXRuntimeORT/Private/NNXRuntimeORTUtils.h b/Engine/Plugins/Experimental/NNX/Source/NNXRuntimeORT/Private/NNXRuntimeORTUtils.h
--- a/Engine/Plugins/Experimental/NNX/Source/NNXRuntimeORT/Private/NNXRuntimeORTUtils.h
+++ b/Engine/Plugins/Experimental/NNX/Source/NNXRuntimeORT/Private/NNXRuntimeORTUtils.h
@@ -194,7 +194,7 @@ inline void CopyFromORTToBindings(
 		const void* OrtTensorData = OrtTensor.GetTensorData<void>();
 		const uint64 DataSize = FMath::Min(Tensor.GetDataSize(), Binding.SizeInBytes);
 
-		if (DataSize >= 0)
+		if (DataSize > 0)
 		{
 			FMemory::Memcpy(CpuMemory, OrtTensorData, DataSize);
 		}
