NNI Dynamic shape support for HLSL engine, preparation CL

Problem:
* Models with dynamic input shapes means intermediate tensor shapes can change at runtime.

Solution:
* Detect that a network is using dynamic shapes and delegate the output shape computation to the op in that case.
* This CL a preparation work and only static network work for now, as no operator is implementing shape inference yet.

Tests:
* All NNX tests pass under win64 editor

diff --git a/Engine/Plugins/Experimental/NNX/Source/NNXRuntimeCPU/Private/NNXRuntimeCPU.cpp b/Engine/Plugins/Experimental/NNX/Source/NNXRuntimeCPU/Private/NNXRuntimeCPU.cpp
--- a/Engine/Plugins/Experimental/NNX/Source/NNXRuntimeCPU/Private/NNXRuntimeCPU.cpp
+++ b/Engine/Plugins/Experimental/NNX/Source/NNXRuntimeCPU/Private/NNXRuntimeCPU.cpp
@@ -201,6 +201,10 @@ bool FMLInferenceModelCPU::ConfigureTensors(const bool InIsInput)
 
 int FMLInferenceModelCPU::SetInputTensorShapes(TConstArrayView<FTensorShape> InInputShapes)
 {
+	InputTensors.Empty();
+	OutputTensors.Empty();
+	OutputTensorShapes.Empty();
+	
 	// Verify input shape are valid for the model and set InputTensorShapes
 	if (FMLInferenceModel::SetInputTensorShapes(InInputShapes) != 0)
 	{
@@ -208,7 +212,6 @@ int FMLInferenceModelCPU::SetInputTensorShapes(TConstArrayView<FTensorShape> InI
 	}
 
 	// Setup concrete input tensor
-	InputTensors.Empty();
 	for (int i = 0; i < InputSymbolicTensors.Num(); ++i)
 	{
 		FTensor Tensor = FTensor::Make(InputSymbolicTensors[i].GetName(), InInputShapes[i], InputSymbolicTensors[i].GetDataType());
@@ -219,7 +222,6 @@ int FMLInferenceModelCPU::SetInputTensorShapes(TConstArrayView<FTensorShape> InI
 	// this would allow to resolve output shapes here rather than during inference.
 
 	// Setup concrete output shapes only if all model output shapes are concretes, otherwise it will be set during Run()
-	OutputTensors.Empty();
 	for (FTensorDesc SymbolicTensorDesc : OutputSymbolicTensors)
 	{
 		if (SymbolicTensorDesc.IsConcrete())
diff --git a/Engine/Plugins/Experimental/NNX/Source/NNXRuntimeORT/Private/NNXRuntimeORT.cpp b/Engine/Plugins/Experimental/NNX/Source/NNXRuntimeORT/Private/NNXRuntimeORT.cpp
--- a/Engine/Plugins/Experimental/NNX/Source/NNXRuntimeORT/Private/NNXRuntimeORT.cpp
+++ b/Engine/Plugins/Experimental/NNX/Source/NNXRuntimeORT/Private/NNXRuntimeORT.cpp
@@ -275,6 +275,10 @@ bool FMLInferenceModelORT::ConfigureTensors(const bool InIsInput)
 
 int FMLInferenceModelORT::SetInputTensorShapes(TConstArrayView<FTensorShape> InInputShapes)
 {
+	InputTensors.Empty();
+	OutputTensors.Empty();
+	OutputTensorShapes.Empty();
+	
 	// Verify input shape are valid for the model and set InputTensorShapes
 	if (FMLInferenceModel::SetInputTensorShapes(InInputShapes) != 0)
 	{
@@ -282,7 +286,6 @@ int FMLInferenceModelORT::SetInputTensorShapes(TConstArrayView<FTensorShape> InI
 	}
 
 	// Setup concrete input tensor
-	InputTensors.Empty();
 	for (int i = 0; i < InputSymbolicTensors.Num(); ++i)
 	{
 		FTensor Tensor = FTensor::Make(InputSymbolicTensors[i].GetName(), InInputShapes[i], InputSymbolicTensors[i].GetDataType());
@@ -293,7 +296,6 @@ int FMLInferenceModelORT::SetInputTensorShapes(TConstArrayView<FTensorShape> InI
 	// this would allow to resolve output shapes here rather than during inference.
 
 	// Setup concrete output shapes only if all model output shapes are concretes, otherwise it will be set during Run()
-	OutputTensors.Empty();
 	for (FTensorDesc SymbolicTensorDesc : OutputSymbolicTensors)
 	{
 		if (SymbolicTensorDesc.IsConcrete())
diff --git a/Engine/Plugins/Experimental/NNX/Source/NNXRuntimeRDG/Private/NNXRuntimeDML.cpp b/Engine/Plugins/Experimental/NNX/Source/NNXRuntimeRDG/Private/NNXRuntimeDML.cpp
--- a/Engine/Plugins/Experimental/NNX/Source/NNXRuntimeRDG/Private/NNXRuntimeDML.cpp
+++ b/Engine/Plugins/Experimental/NNX/Source/NNXRuntimeRDG/Private/NNXRuntimeDML.cpp
@@ -1173,6 +1173,7 @@ public:
 protected:
 
 	virtual void AddDispatchOps_RenderThread(FRDGBuilder& GraphBuilder) override;
+	virtual int RunShapeInference() override;
 
 private:
 
@@ -1581,6 +1582,31 @@ FMLOperatorDml* FMLInferenceModelDml::OpCreate(const FString& OpName, TArrayView
 	return Op;
 }
 
+//
+//
+//
+int FMLInferenceModelDml::RunShapeInference()
+{
+	AllTensors.Empty();
+
+	for (FTensorDesc SymbolicTensorDesc : AllSymbolicTensors)
+	{
+		if (SymbolicTensorDesc.IsConcrete())
+		{
+			FTensor Tensor = FTensor::MakeFromSymbolicDesc(SymbolicTensorDesc);
+			AllTensors.Emplace(Tensor);
+		}
+	}
+	if (AllTensors.Num() != AllSymbolicTensors.Num())
+	{
+		AllTensors.Empty();
+		UE_LOG(LogNNX, Warning, TEXT("DML engine does not support model with variable shapes yet."));
+		return -1;
+	}
+
+	return 0;
+}
+
 //
 //
 //
diff --git a/Engine/Plugins/Experimental/NNX/Source/NNXRuntimeRDG/Private/NNXRuntimeHLSL.cpp b/Engine/Plugins/Experimental/NNX/Source/NNXRuntimeRDG/Private/NNXRuntimeHLSL.cpp
--- a/Engine/Plugins/Experimental/NNX/Source/NNXRuntimeRDG/Private/NNXRuntimeHLSL.cpp
+++ b/Engine/Plugins/Experimental/NNX/Source/NNXRuntimeRDG/Private/NNXRuntimeHLSL.cpp
@@ -35,6 +35,7 @@ public:
 protected:
 
 	virtual void AddDispatchOps_RenderThread(FRDGBuilder& GraphBuilder) override;
+	virtual int RunShapeInference() override;
 
 private:
 
@@ -134,6 +135,107 @@ FMLOperatorHlsl* FMLInferenceModelHlsl::OpCreate(const FString& OpName, TArrayVi
 	return Op;
 }
 
+int FMLInferenceModelHlsl::RunShapeInference()
+{
+	AllTensors.Empty();
+
+	if (Operators.Num() == 0)
+	{
+		UE_LOG(LogNNX, Warning, TEXT("No operators in model"));
+		return -1;
+	}
+
+	static constexpr int32 MaxExpectedInput = 10;
+	TArray<FTensorShape, TInlineAllocator<MaxExpectedInput>> InputShapes;
+	TArray<FTensorShape> OutputShapes;
+	TArray<FTensorShape> AllShapes;
+	TArray<bool> AllInitializedShapes;
+
+	checkCode(AllInitializedShapes.SetNum(AllSymbolicTensors.Num(), false););
+	AllShapes.SetNum(AllSymbolicTensors.Num());
+
+	//Prime shape inference with model inputs
+	for (int32 i = 0; i < InputTensorIndices.Num(); ++i)
+	{
+		const int32 InputTensorModelIndex = InputTensorIndices[i];
+
+		AllShapes[InputTensorModelIndex] = InputTensorShapes[i];
+		checkCode(AllInitializedShapes[InputTensorModelIndex] = true);
+	}
+
+	// Run shape inference on all operators
+	for (int32 Idx = 0; Idx < Operators.Num(); ++Idx)
+	{
+		InputShapes.Empty();
+		OutputShapes.Empty();
+
+		for (int32 i : OperatorInputTensorIndices[Idx])
+		{
+			checkf(AllInitializedShapes[i] == true, TEXT("Input shape %d for operator %d should have been initialized."), i, Idx);
+			InputShapes.Emplace(AllShapes[i]);
+		}
+
+		//If operator output shapes are not variable we don't need to run shape inference
+		bool bShouldRunShapeInferenceForOperator = false;
+		for (int32 i = 0; i < OperatorOutputTensorIndices[Idx].Num(); ++i)
+		{
+			const int32 OutputTensorIndex = OperatorOutputTensorIndices[Idx][i];
+
+			if (!OutputSymbolicTensors[i].IsConcrete())
+			{
+				bShouldRunShapeInferenceForOperator = true;
+				continue;
+			}
+			OutputShapes.Emplace(FTensorShape::MakeFromSymbolic(OutputSymbolicTensors[i].GetShape()));
+		}
+
+		//Otherwise we need to run shape inference for the operator
+		if (bShouldRunShapeInferenceForOperator)
+		{
+			OutputShapes.Empty();
+
+			const FMLOperatorHlsl* Op = Operators[Idx];
+
+			if (Op->ComputeOutputShape(InputShapes, OutputShapes) != 0)
+			{
+				//Operator could not compute output shapes, meaning we can't allocate
+				//output buffer before running the model. This engine does not support this.
+				UE_LOG(LogNNX, Warning, TEXT("Could not deduce tensor shapes for this model during shape inference, HLSL engine wont support the model as it need to precompute all shapes for performance reasons."));
+				AllTensors.Empty();
+				OutputTensorShapes.Empty();
+				return -1;
+			}
+		}
+
+		for (int32 i = 0; i < OperatorOutputTensorIndices[Idx].Num(); ++i)
+		{
+			const int32 OutputTensorIndex = OperatorOutputTensorIndices[Idx][i];
+
+			AllShapes[OutputTensorIndex] = OutputShapes[i];
+			checkCode(AllInitializedShapes[OutputTensorIndex] = true);
+		}
+	}
+
+	checkCode(
+		for (int i = 0; i < AllInitializedShapes.Num(); ++i)
+		{
+			checkf(AllInitializedShapes[i], TEXT("Tensor at index %d, was not initialized by shape inference."));
+		};
+	);
+
+	AllTensors.Reserve(AllShapes.Num());
+	for (int i = 0; i < AllShapes.Num(); ++i)
+	{
+		FTensor Tensor = FTensor::Make(AllSymbolicTensors[i].GetName(), AllShapes[i], AllSymbolicTensors[i].GetDataType());
+		AllTensors.Emplace(Tensor);
+	}
+
+	check(AllShapes.Num() == AllSymbolicTensors.Num());
+	check(AllShapes.Num() == AllTensors.Num());
+
+	return 0;
+}
+
 //
 //
 //
diff --git a/Engine/Plugins/Experimental/NNX/Source/NNXRuntimeRDG/Private/NNXRuntimeHLSLOp.h b/Engine/Plugins/Experimental/NNX/Source/NNXRuntimeRDG/Private/NNXRuntimeHLSLOp.h
--- a/Engine/Plugins/Experimental/NNX/Source/NNXRuntimeRDG/Private/NNXRuntimeHLSLOp.h
+++ b/Engine/Plugins/Experimental/NNX/Source/NNXRuntimeRDG/Private/NNXRuntimeHLSLOp.h
@@ -15,6 +15,7 @@ namespace NNX
 		virtual ~FMLOperatorHlsl() = default;
 		virtual bool Initialize(TArrayView<const FTensor> InputTensors, TArrayView<const FTensor> OutputTensors, const UE::NNECore::FAttributeMap& Attributes) = 0;
 		virtual void Dispatch(FRDGBuilder& GraphBuilder, TArrayView<const FMLTensorBinding> InInputBindings, TArrayView<const FMLTensorBinding> OutOutputBindings) = 0;
+		virtual int ComputeOutputShape(TConstArrayView<FTensorShape> InputShapes, TArray<FTensorShape>& OutputShapes) const { return -1; };
 	};
 
 	/**
diff --git a/Engine/Plugins/Experimental/NNX/Source/NNXRuntimeRDG/Private/NNXRuntimeRDG.cpp b/Engine/Plugins/Experimental/NNX/Source/NNXRuntimeRDG/Private/NNXRuntimeRDG.cpp
--- a/Engine/Plugins/Experimental/NNX/Source/NNXRuntimeRDG/Private/NNXRuntimeRDG.cpp
+++ b/Engine/Plugins/Experimental/NNX/Source/NNXRuntimeRDG/Private/NNXRuntimeRDG.cpp
@@ -307,31 +307,33 @@ int FMLInferenceModelRDG::Run(TConstArrayView<FMLTensorBinding> InInputBindings,
 
 int FMLInferenceModelRDG::SetInputTensorShapes(TConstArrayView<FTensorShape> InInputShapes)
 {
+	InputTensors.Empty();
+	OutputTensors.Empty();
+	OutputTensorShapes.Empty();
+	
 	//Verify input shape are valid for the model and set InputTensorShapes
 	if (FMLInferenceModel::SetInputTensorShapes(InInputShapes) != 0)
 	{
 		return -1;
 	}
 
-	//Run shape inference
-	AllTensors.Empty();
-	for (FTensorDesc SymbolicTensorDesc : AllSymbolicTensors)
+	//Run shape inference filling AllTensors
+	if (RunShapeInference() != 0)
 	{
-		//Fake shape inference for now
-		check(SymbolicTensorDesc.IsConcrete());
-		FTensor TensorDesc = FTensor::MakeFromSymbolicDesc(SymbolicTensorDesc);
-		AllTensors.Emplace(TensorDesc);
+		return -1;
 	}
 
-	//Build concrete input and output tensor descriptions
-	InputTensors.Empty();
-	for (int32 InputIndices : InputTensorIndices)
+	// Setup concrete input tensor
+	for (int32 i = 0; i < InputTensorIndices.Num(); ++i)
 	{
-		InputTensors.Emplace(AllTensors[InputIndices]);
+		const int32 TensorIndice = InputTensorIndices[i];
+		const FTensor& InputTensor = AllTensors[TensorIndice];
+
+		InputTensors.Emplace(InputTensor);
+		check(InputTensor.GetShape().Data == InputTensorShapes[i].Data);
 	}
-	
-	OutputTensors.Empty();
-	OutputTensorShapes.Empty();
+
+	//Set OutputTensorShapes for the model from shape inference result
 	for (int32 OutputIndices : OutputTensorIndices)
 	{
 		OutputTensors.Emplace(AllTensors[OutputIndices]);
diff --git a/Engine/Plugins/Experimental/NNX/Source/NNXRuntimeRDG/Private/NNXRuntimeRDG.h b/Engine/Plugins/Experimental/NNX/Source/NNXRuntimeRDG/Private/NNXRuntimeRDG.h
--- a/Engine/Plugins/Experimental/NNX/Source/NNXRuntimeRDG/Private/NNXRuntimeRDG.h
+++ b/Engine/Plugins/Experimental/NNX/Source/NNXRuntimeRDG/Private/NNXRuntimeRDG.h
@@ -74,7 +74,8 @@ protected:
 	bool LoadModel(const FNNIModelRaw& InModel, FMLRuntimeFormat& Format);
 
 	int SetTensors(FRDGBuilder& GraphBuilder, FMLTensorBindingArray& OutRDGBindings, FMLIntArray& OutIndices, TArrayView<const FMLTensorBinding> InBindings, TArrayView<const FTensor> InTensors);
-
+	
+	virtual int RunShapeInference() = 0;
 	virtual void AddDispatchOps_RenderThread(FRDGBuilder& GraphBuilder) = 0;
 
 	virtual void AddTensorUploads_RenderThread(FRDGBuilder& GraphBuilder, TArrayView<const int32> InUploadIndices, TArrayView<FMLTensorBinding> InRDGBindings, TArrayView<const FMLTensorBinding> InBindings);
diff --git a/Engine/Plugins/Experimental/NNX/Source/NNXUtils/Private/NNXModelBuilderNNX.cpp b/Engine/Plugins/Experimental/NNX/Source/NNXUtils/Private/NNXModelBuilderNNX.cpp
--- a/Engine/Plugins/Experimental/NNX/Source/NNXUtils/Private/NNXModelBuilderNNX.cpp
+++ b/Engine/Plugins/Experimental/NNX/Source/NNXUtils/Private/NNXModelBuilderNNX.cpp
@@ -111,7 +111,14 @@ public:
 
 	virtual HTensor AddTensor(const FString& Name, EMLTensorDataType DataType, TArrayView<const int32> Shape)
 	{
-		int Idx = AddTensor(Name, Shape, DataType, EMLFormatTensorType::None);
+		TArray<int32, TInlineAllocator<FTensorShape::MaxRank>> NNIShape;
+		for (int i = 0; i < Shape.Num(); ++i)
+		{
+			//ORT Graph return 0 for variable dimensions, NNI use -1.
+			NNIShape.Emplace(Shape[i] == 0 ? -1 : Shape[i]);
+		}
+		
+		int Idx = AddTensor(Name, NNIShape, DataType, EMLFormatTensorType::None);
 		
 		return MakeTensorHandle(reinterpret_cast<void*>((int64) Idx));
 	}
