NNI Dynamic shape support for HLSL engine, Proof of concept with ElementWise operator

Problem:
* Operators are initialized from concrete shape, preventing to support dynamic shapes

Solution:
* Initialize operator from tensor description and handle the case where initialization can fail.
* Implement dynamic support for element wise op as a proof of concept

Tests:
* All NNX tests pass under win64 editor, including dynamic model using element wise op.

diff --git a/Engine/Plugins/Experimental/NNX/Source/NNXQA/Resources/NNXQATestDesc.json b/Engine/Plugins/Experimental/NNX/Source/NNXQA/Resources/NNXQATestDesc.json
--- a/Engine/Plugins/Experimental/NNX/Source/NNXQA/Resources/NNXQATestDesc.json
+++ b/Engine/Plugins/Experimental/NNX/Source/NNXQA/Resources/NNXQATestDesc.json
@@ -88,7 +88,6 @@
 					"target": "OneLayerNoWeightVariableInputShape",
 					"tags": [ "SmokeTest" ],
 					"runtimes": [
-						{ "name": "NNXRuntimeHLSL", "skip": true, "comment": "RDG does not support variable input sized tensor" },
 						{ "name": "NNXRuntimeDML",  "skip": true, "comment": "RDG does not support variable input sized tensor" }
 					]
 				},
@@ -96,7 +95,6 @@
 					"target": "TwoLayersNoWeightVariableInputShape",
 					"tags": [ "SmokeTest" ],
 					"runtimes": [
-						{ "name": "NNXRuntimeHLSL", "skip": true, "comment": "RDG does not support variable input sized tensor" },
 						{ "name": "NNXRuntimeDML",  "skip": true, "comment": "RDG does not support variable input sized tensor" }
 					]
 				}
diff --git a/Engine/Plugins/Experimental/NNX/Source/NNXRuntimeRDG/Private/Hlsl/NNIRuntimeRDGConv.cpp b/Engine/Plugins/Experimental/NNX/Source/NNXRuntimeRDG/Private/Hlsl/NNIRuntimeRDGConv.cpp
--- a/Engine/Plugins/Experimental/NNX/Source/NNXRuntimeRDG/Private/Hlsl/NNIRuntimeRDGConv.cpp
+++ b/Engine/Plugins/Experimental/NNX/Source/NNXRuntimeRDG/Private/Hlsl/NNIRuntimeRDGConv.cpp
@@ -46,8 +46,17 @@ namespace UE::NNIRuntimeRDG::Private::Hlsl
 
 	public:
 
-		virtual bool Initialize(TArrayView<const NNX::FTensor> InputTensors, TArrayView<const NNX::FTensor> OutputTensors, const UE::NNECore::FAttributeMap& Attributes) override
+		virtual bool Initialize(TConstArrayView<NNX::FTensorDesc> InputTensorDescs, TConstArrayView<NNX::FTensorDesc> OutputTensorDescs, const UE::NNECore::FAttributeMap& Attributes) override
 		{
+			TArray<NNX::FTensor> InputTensors;
+			TArray<NNX::FTensor> OutputTensors;
+			if (!NNX::ConvertConcreteTensorDescsToTensors(InputTensorDescs, InputTensors) ||
+				!NNX::ConvertConcreteTensorDescsToTensors(OutputTensorDescs, OutputTensors))
+			{
+				UE_LOG(LogNNX, Warning, TEXT("Variable input shapes are not supported by this operator"));
+				return false;
+			}
+			
 			check(InputTensors.Num() >= 2 && InputTensors.Num() <= 3);
 			check(OutputTensors.Num() == 1);
 
diff --git a/Engine/Plugins/Experimental/NNX/Source/NNXRuntimeRDG/Private/Hlsl/NNIRuntimeRDGConvTranspose.cpp b/Engine/Plugins/Experimental/NNX/Source/NNXRuntimeRDG/Private/Hlsl/NNIRuntimeRDGConvTranspose.cpp
--- a/Engine/Plugins/Experimental/NNX/Source/NNXRuntimeRDG/Private/Hlsl/NNIRuntimeRDGConvTranspose.cpp
+++ b/Engine/Plugins/Experimental/NNX/Source/NNXRuntimeRDG/Private/Hlsl/NNIRuntimeRDGConvTranspose.cpp
@@ -46,8 +46,17 @@ namespace UE::NNIRuntimeRDG::Private::Hlsl
 
 	public:
 
-		virtual bool Initialize(TArrayView<const NNX::FTensor> InputTensors, TArrayView<const NNX::FTensor> OutputTensors, const UE::NNECore::FAttributeMap& Attributes) override
+		virtual bool Initialize(TConstArrayView<NNX::FTensorDesc> InputTensorDescs, TConstArrayView<NNX::FTensorDesc> OutputTensorDescs, const UE::NNECore::FAttributeMap& Attributes) override
 		{
+			TArray<NNX::FTensor> InputTensors;
+			TArray<NNX::FTensor> OutputTensors;
+			if (!NNX::ConvertConcreteTensorDescsToTensors(InputTensorDescs, InputTensors) ||
+				!NNX::ConvertConcreteTensorDescsToTensors(OutputTensorDescs, OutputTensors))
+			{
+				UE_LOG(LogNNX, Warning, TEXT("Variable input shapes are not supported by this operator"));
+				return false;
+			}
+				
 			check(InputTensors.Num() >= 2 && InputTensors.Num() <= 3);
 			check(OutputTensors.Num() == 1);
 
diff --git a/Engine/Plugins/Experimental/NNX/Source/NNXRuntimeRDG/Private/Hlsl/NNIRuntimeRDGElementWiseBinary.cpp b/Engine/Plugins/Experimental/NNX/Source/NNXRuntimeRDG/Private/Hlsl/NNIRuntimeRDGElementWiseBinary.cpp
--- a/Engine/Plugins/Experimental/NNX/Source/NNXRuntimeRDG/Private/Hlsl/NNIRuntimeRDGElementWiseBinary.cpp
+++ b/Engine/Plugins/Experimental/NNX/Source/NNXRuntimeRDG/Private/Hlsl/NNIRuntimeRDGElementWiseBinary.cpp
@@ -30,8 +30,17 @@ namespace UE::NNIRuntimeRDG::Private::Hlsl
 
 	public:
 
-		virtual bool Initialize(TArrayView<const NNX::FTensor> InputTensors, TArrayView<const NNX::FTensor> OutputTensors, const UE::NNECore::FAttributeMap& Attributes) override
+		virtual bool Initialize(TConstArrayView<NNX::FTensorDesc> InputTensorDescs, TConstArrayView<NNX::FTensorDesc> OutputTensorDescs, const UE::NNECore::FAttributeMap& Attributes) override
 		{
+			TArray<NNX::FTensor> InputTensors;
+			TArray<NNX::FTensor> OutputTensors;
+			if (!NNX::ConvertConcreteTensorDescsToTensors(InputTensorDescs, InputTensors) ||
+				!NNX::ConvertConcreteTensorDescsToTensors(OutputTensorDescs, OutputTensors))
+			{
+				UE_LOG(LogNNX, Warning, TEXT("Variable input shapes are not supported by this operator"));
+				return false;
+			}
+			
 			check(InputTensors.Num() == 2);
 			check(OutputTensors.Num() == 1);
 		
diff --git a/Engine/Plugins/Experimental/NNX/Source/NNXRuntimeRDG/Private/Hlsl/NNIRuntimeRDGElementWiseUnary.cpp b/Engine/Plugins/Experimental/NNX/Source/NNXRuntimeRDG/Private/Hlsl/NNIRuntimeRDGElementWiseUnary.cpp
--- a/Engine/Plugins/Experimental/NNX/Source/NNXRuntimeRDG/Private/Hlsl/NNIRuntimeRDGElementWiseUnary.cpp
+++ b/Engine/Plugins/Experimental/NNX/Source/NNXRuntimeRDG/Private/Hlsl/NNIRuntimeRDGElementWiseUnary.cpp
@@ -28,23 +28,26 @@ namespace UE::NNIRuntimeRDG::Private::Hlsl
 		float Alpha = 0.0f;
 		float Beta = 0.0f;
 		float Gamma = 0.0f;
-		NNX::FTensor Input = {};
-		NNX::FTensor Output = {};
+		int32 ElemByteSize = 0;
 
 	public:
 
-		virtual bool Initialize(TArrayView<const NNX::FTensor> InputTensors, TArrayView<const NNX::FTensor> OutputTensors, const UE::NNECore::FAttributeMap& Attributes) override
+		virtual int ComputeOutputShape(TConstArrayView<NNX::FTensorShape> InputShapes, TArray<NNX::FTensorShape>& OutputShapes) const override
 		{
-			check(InputTensors.Num() == 1);
-			check(OutputTensors.Num() == 1);
+			check(InputShapes.Num() == 1);
+			OutputShapes = InputShapes;
+			return 0;
+		}
 
-			Input = InputTensors[0];
-			Output = OutputTensors[0];
+		virtual bool Initialize(TConstArrayView<NNX::FTensorDesc> InputTensorDescs, TConstArrayView<NNX::FTensorDesc> OutputTensorDescs, const UE::NNECore::FAttributeMap& Attributes) override
+		{
+			check(InputTensorDescs.Num() == 1);
+			check(OutputTensorDescs.Num() == 1);
 
 			Alpha = Attributes.GetValueOrDefault(TEXT("alpha"), Alpha);
 			Beta = Attributes.GetValueOrDefault(TEXT("beta"), Beta);
 			Gamma = Attributes.GetValueOrDefault(TEXT("gamma"), Gamma);
-
+			ElemByteSize = InputTensorDescs[0].GetElemByteSize();
 			return true;
 		}
 
@@ -78,7 +81,8 @@ namespace UE::NNIRuntimeRDG::Private::Hlsl
 			FRDGBufferSRVRef InputSRV = GraphBuilder.CreateSRV(FRDGBufferSRVDesc(InInputBindings[0].Buffer, PF_R32_FLOAT));
 			FRDGBufferUAVRef OutputUAV = GraphBuilder.CreateUAV(FRDGBufferUAVDesc(OutOutputBindings[0].Buffer, PF_R32_FLOAT));
 		
-			FIntVector ThreadGroupCount = NNX::ComputeElementWiseThreadGroups(Output.GetVolume(), FElementWiseUnaryConstants::NUM_GROUP_THREADS);
+			int32 NumElements = OutOutputBindings[0].SizeInBytes / ElemByteSize;//TODO jira 169354: Use shape from Dispatch instead
+			FIntVector ThreadGroupCount = NNX::ComputeElementWiseThreadGroups(NumElements, FElementWiseUnaryConstants::NUM_GROUP_THREADS);
 
 			// Set parameters
 			TElementWiseUnaryCS::FParameters* Params = GraphBuilder.AllocParameters<TElementWiseUnaryCS::FParameters>();
@@ -87,7 +91,7 @@ namespace UE::NNIRuntimeRDG::Private::Hlsl
 			Params->Alpha = Alpha;
 			Params->Beta = Beta;
 			Params->Gamma = Gamma;
-			Params->Num = Output.GetVolume();
+			Params->Num = NumElements;
 			Params->ThreadCountX = ThreadGroupCount.X * FElementWiseUnaryConstants::NUM_GROUP_THREADS;
 
 			TElementWiseUnaryCS::FPermutationDomain PermutationVector;
diff --git a/Engine/Plugins/Experimental/NNX/Source/NNXRuntimeRDG/Private/Hlsl/NNIRuntimeRDGElementWiseVariadic.cpp b/Engine/Plugins/Experimental/NNX/Source/NNXRuntimeRDG/Private/Hlsl/NNIRuntimeRDGElementWiseVariadic.cpp
--- a/Engine/Plugins/Experimental/NNX/Source/NNXRuntimeRDG/Private/Hlsl/NNIRuntimeRDGElementWiseVariadic.cpp
+++ b/Engine/Plugins/Experimental/NNX/Source/NNXRuntimeRDG/Private/Hlsl/NNIRuntimeRDGElementWiseVariadic.cpp
@@ -99,8 +99,17 @@ namespace UE::NNIRuntimeRDG::Private::Hlsl
 
 	public:
 
-		virtual bool Initialize(TArrayView<const NNX::FTensor> InputTensors, TArrayView<const NNX::FTensor> OutputTensors, const UE::NNECore::FAttributeMap& Attributes) override
+		virtual bool Initialize(TConstArrayView<NNX::FTensorDesc> InputTensorDescs, TConstArrayView<NNX::FTensorDesc> OutputTensorDescs, const UE::NNECore::FAttributeMap& Attributes) override
 		{
+			TArray<NNX::FTensor> InputTensors;
+			TArray<NNX::FTensor> OutputTensors;
+			if (!NNX::ConvertConcreteTensorDescsToTensors(InputTensorDescs, InputTensors) ||
+				!NNX::ConvertConcreteTensorDescsToTensors(OutputTensorDescs, OutputTensors))
+			{
+				UE_LOG(LogNNX, Warning, TEXT("Variable input shapes are not supported by this operator"));
+				return false;
+			}
+			
 			check(InputTensors.Num() > 0);
 			check(OutputTensors.Num() == 1);
 		
diff --git a/Engine/Plugins/Experimental/NNX/Source/NNXRuntimeRDG/Private/Hlsl/NNIRuntimeRDGGather.cpp b/Engine/Plugins/Experimental/NNX/Source/NNXRuntimeRDG/Private/Hlsl/NNIRuntimeRDGGather.cpp
--- a/Engine/Plugins/Experimental/NNX/Source/NNXRuntimeRDG/Private/Hlsl/NNIRuntimeRDGGather.cpp
+++ b/Engine/Plugins/Experimental/NNX/Source/NNXRuntimeRDG/Private/Hlsl/NNIRuntimeRDGGather.cpp
@@ -29,8 +29,17 @@ namespace UE::NNIRuntimeRDG::Private::Hlsl
 
 	public:
 
-		virtual bool Initialize(TArrayView<const NNX::FTensor> InputTensors, TArrayView<const NNX::FTensor> OutputTensors, const UE::NNECore::FAttributeMap& Attributes) override
+		virtual bool Initialize(TConstArrayView<NNX::FTensorDesc> InputTensorDescs, TConstArrayView<NNX::FTensorDesc> OutputTensorDescs, const UE::NNECore::FAttributeMap& Attributes) override
 		{
+			TArray<NNX::FTensor> InputTensors;
+			TArray<NNX::FTensor> OutputTensors;
+			if (!NNX::ConvertConcreteTensorDescsToTensors(InputTensorDescs, InputTensors) ||
+				!NNX::ConvertConcreteTensorDescsToTensors(OutputTensorDescs, OutputTensors))
+			{
+				UE_LOG(LogNNX, Warning, TEXT("Variable input shapes are not supported by this operator"));
+				return false;
+			}
+			
 			using namespace UE::NNEHlslShaders::Internal;
 
 			check(InputTensors.Num() == 2)
diff --git a/Engine/Plugins/Experimental/NNX/Source/NNXRuntimeRDG/Private/Hlsl/NNIRuntimeRDGGemm.cpp b/Engine/Plugins/Experimental/NNX/Source/NNXRuntimeRDG/Private/Hlsl/NNIRuntimeRDGGemm.cpp
--- a/Engine/Plugins/Experimental/NNX/Source/NNXRuntimeRDG/Private/Hlsl/NNIRuntimeRDGGemm.cpp
+++ b/Engine/Plugins/Experimental/NNX/Source/NNXRuntimeRDG/Private/Hlsl/NNIRuntimeRDGGemm.cpp
@@ -41,8 +41,17 @@ namespace UE::NNIRuntimeRDG::Private::Hlsl
 
 	public:
 
-		virtual bool Initialize(TArrayView<const NNX::FTensor> InputTensors, TArrayView<const NNX::FTensor> OutputTensors, const UE::NNECore::FAttributeMap& Attributes) override
+		virtual bool Initialize(TConstArrayView<NNX::FTensorDesc> InputTensorDescs, TConstArrayView<NNX::FTensorDesc> OutputTensorDescs, const UE::NNECore::FAttributeMap& Attributes) override
 		{
+			TArray<NNX::FTensor> InputTensors;
+			TArray<NNX::FTensor> OutputTensors;
+			if (!NNX::ConvertConcreteTensorDescsToTensors(InputTensorDescs, InputTensors) ||
+				!NNX::ConvertConcreteTensorDescsToTensors(OutputTensorDescs, OutputTensors))
+			{
+				UE_LOG(LogNNX, Warning, TEXT("Variable input shapes are not supported by this operator"));
+				return false;
+			}
+			
 			check(InputTensors.Num() >= 2 && InputTensors.Num() <= 3);
 			check(OutputTensors.Num() == 1);
 
diff --git a/Engine/Plugins/Experimental/NNX/Source/NNXRuntimeRDG/Private/Hlsl/NNIRuntimeRDGMatMul.cpp b/Engine/Plugins/Experimental/NNX/Source/NNXRuntimeRDG/Private/Hlsl/NNIRuntimeRDGMatMul.cpp
--- a/Engine/Plugins/Experimental/NNX/Source/NNXRuntimeRDG/Private/Hlsl/NNIRuntimeRDGMatMul.cpp
+++ b/Engine/Plugins/Experimental/NNX/Source/NNXRuntimeRDG/Private/Hlsl/NNIRuntimeRDGMatMul.cpp
@@ -26,8 +26,17 @@ namespace UE::NNIRuntimeRDG::Private::Hlsl
 
 	public:
 
-		virtual bool Initialize(TArrayView<const NNX::FTensor> InputTensors, TArrayView<const NNX::FTensor> OutputTensors, const UE::NNECore::FAttributeMap& Attributes) override
+		virtual bool Initialize(TConstArrayView<NNX::FTensorDesc> InputTensorDescs, TConstArrayView<NNX::FTensorDesc> OutputTensorDescs, const UE::NNECore::FAttributeMap& Attributes) override
 		{
+			TArray<NNX::FTensor> InputTensors;
+			TArray<NNX::FTensor> OutputTensors;
+			if (!NNX::ConvertConcreteTensorDescsToTensors(InputTensorDescs, InputTensors) ||
+				!NNX::ConvertConcreteTensorDescsToTensors(OutputTensorDescs, OutputTensors))
+			{
+				UE_LOG(LogNNX, Warning, TEXT("Variable input shapes are not supported by this operator"));
+				return false;
+			}
+			
 			check(InputTensors.Num() == 2);
 			check(OutputTensors.Num() == 1);
 
diff --git a/Engine/Plugins/Experimental/NNX/Source/NNXRuntimeRDG/Private/NNXRuntimeHLSL.cpp b/Engine/Plugins/Experimental/NNX/Source/NNXRuntimeRDG/Private/NNXRuntimeHLSL.cpp
--- a/Engine/Plugins/Experimental/NNX/Source/NNXRuntimeRDG/Private/NNXRuntimeHLSL.cpp
+++ b/Engine/Plugins/Experimental/NNX/Source/NNXRuntimeRDG/Private/NNXRuntimeHLSL.cpp
@@ -39,7 +39,7 @@ protected:
 
 private:
 
-	FMLOperatorHlsl* OpCreate(const FString& Name, TArrayView<const FTensor> InputTensorDesc, TArrayView<const FTensor> OutputTensorDescs, const UE::NNECore::FAttributeMap& Attributes);
+	FMLOperatorHlsl* OpCreate(const FString& Name, TConstArrayView<FTensorDesc> InputTensorDesc, TConstArrayView<FTensorDesc> OutputTensorDescs, const UE::NNECore::FAttributeMap& Attributes);
 
 	TArray<FMLOperatorHlsl*>	Operators;
 };
@@ -78,28 +78,24 @@ bool FMLInferenceModelHlsl::Init(UMLInferenceModel* InModel)
 	{
 		const FString TypeName = Format.Operators[Idx].TypeName;
 
-		TArray<FTensor> OpInputTensors;
-		TArray<FTensor> OpOutputTensors;
+		TArray<FTensorDesc> Inputs;
+		TArray<FTensorDesc> Outputs;
 		UE::NNECore::FAttributeMap AttributeMap;
 
 		for (int32 InputTensorIndex : Format.Operators[Idx].InTensors)
 		{
-			FTensorDesc SymbolicTensorDesc = AllSymbolicTensors[InputTensorIndex];
-			//TODO jira 168972: Handle dynamic tensor desc, op should init from symbolic shapes
-			OpInputTensors.Emplace(FTensor::MakeFromSymbolicDesc(SymbolicTensorDesc));
+			Inputs.Emplace(AllSymbolicTensors[InputTensorIndex]);
 		}
 		for (int32 OutputTensorIndex : Format.Operators[Idx].OutTensors)
 		{
-			FTensorDesc SymbolicTensorDesc = AllSymbolicTensors[OutputTensorIndex];
-			//TODO jira 168972: Handle dynamic tensor desc, op should init from symbolic shapes
-			OpOutputTensors.Emplace(FTensor::MakeFromSymbolicDesc(SymbolicTensorDesc));
+			Outputs.Emplace(AllSymbolicTensors[OutputTensorIndex]);
 		}
 		for (const FMLFormatAttributeDesc& Desc : Format.Operators[Idx].Attributes)
 		{
 			AttributeMap.SetAttribute(Desc.Name, Desc.Value);
 		}
 
-		FMLOperatorHlsl* Op = OpCreate(TypeName, OpInputTensors, OpOutputTensors, AttributeMap);
+		FMLOperatorHlsl* Op = OpCreate(TypeName, Inputs, Outputs, AttributeMap);
 
 		if (!Op) //Op.Shader.IsNull())
 		{
@@ -118,7 +114,7 @@ bool FMLInferenceModelHlsl::Init(UMLInferenceModel* InModel)
 //
 //
 //
-FMLOperatorHlsl* FMLInferenceModelHlsl::OpCreate(const FString& OpName, TArrayView<const FTensor> InputTensorDescs, TArrayView<const FTensor> OutputTensorDescs, const UE::NNECore::FAttributeMap& AttributeMap)
+FMLOperatorHlsl* FMLInferenceModelHlsl::OpCreate(const FString& OpName, TConstArrayView<FTensorDesc> InputTensorDescs, TConstArrayView<FTensorDesc> OutputTensorDescs, const UE::NNECore::FAttributeMap& AttributeMap)
 {
 	FMLOperatorRegistryHlsl::OperatorCreateFunc CreateFn = FMLOperatorRegistryHlsl::Get()->OpFind(OpName);
 
@@ -130,7 +126,12 @@ FMLOperatorHlsl* FMLInferenceModelHlsl::OpCreate(const FString& OpName, TArrayVi
 
 	FMLOperatorHlsl* Op = CreateFn();
 
-	Op->Initialize(InputTensorDescs, OutputTensorDescs, AttributeMap);
+	if (!Op->Initialize(InputTensorDescs, OutputTensorDescs, AttributeMap))
+	{
+		UE_LOG(LogNNX, Warning, TEXT("Hlsl engine: Error initializing operator:%s"), *OpName);
+		delete Op;
+		return nullptr;
+	}
 
 	return Op;
 }
diff --git a/Engine/Plugins/Experimental/NNX/Source/NNXRuntimeRDG/Private/NNXRuntimeHLSLHelper.cpp b/Engine/Plugins/Experimental/NNX/Source/NNXRuntimeRDG/Private/NNXRuntimeHLSLHelper.cpp
--- a/Engine/Plugins/Experimental/NNX/Source/NNXRuntimeRDG/Private/NNXRuntimeHLSLHelper.cpp
+++ b/Engine/Plugins/Experimental/NNX/Source/NNXRuntimeRDG/Private/NNXRuntimeHLSLHelper.cpp
@@ -63,4 +63,16 @@ namespace NNX
 		return ThreadGroupCount;
 	}
 
+	bool ConvertConcreteTensorDescsToTensors(TConstArrayView<FTensorDesc> TensorDescs, TArray<FTensor>& Tensors)
+	{
+		Tensors.Empty();
+		bool bAllShapesWereConcrete = true;
+		for (const FTensorDesc& TensorDesc : TensorDescs)
+		{
+			Tensors.Emplace(FTensor::MakeFromSymbolicDesc(TensorDesc));
+			bAllShapesWereConcrete &= TensorDesc.IsConcrete();
+		}
+		return bAllShapesWereConcrete;
+	}
+
 } // NNX
\ No newline at end of file
diff --git a/Engine/Plugins/Experimental/NNX/Source/NNXRuntimeRDG/Private/NNXRuntimeHLSLHelper.h b/Engine/Plugins/Experimental/NNX/Source/NNXRuntimeRDG/Private/NNXRuntimeHLSLHelper.h
--- a/Engine/Plugins/Experimental/NNX/Source/NNXRuntimeRDG/Private/NNXRuntimeHLSLHelper.h
+++ b/Engine/Plugins/Experimental/NNX/Source/NNXRuntimeRDG/Private/NNXRuntimeHLSLHelper.h
@@ -11,4 +11,5 @@ namespace NNX
 	void FillTensorStrideShaderParameters(const FTensor& TensorDesc, TStaticArray<FUintVector4, NXRT_TENSORSTRIDEINFO_MAX_NUM_DIMENSIONS, 16U>& OutShaderParam, int32 Idx, int32 TargetNumdimensionForBroadcast = -1);
 	void FillTensorStrideForBroadcastShaderParameters(const FTensor& TensorDesc, int32 OutputNumdimension, TStaticArray<FUintVector4, NXRT_TENSORSTRIDEINFO_MAX_NUM_DIMENSIONS, 16U>& OutShaderParam, int32 Idx);
 	FIntVector ComputeElementWiseThreadGroups(uint32 ElementCount, uint32 GroupSizeX);
+	bool ConvertConcreteTensorDescsToTensors(TConstArrayView<FTensorDesc> TensorDescs, TArray<FTensor>& Tensors);
 } // NNX
\ No newline at end of file
diff --git a/Engine/Plugins/Experimental/NNX/Source/NNXRuntimeRDG/Private/NNXRuntimeHLSLOp.h b/Engine/Plugins/Experimental/NNX/Source/NNXRuntimeRDG/Private/NNXRuntimeHLSLOp.h
--- a/Engine/Plugins/Experimental/NNX/Source/NNXRuntimeRDG/Private/NNXRuntimeHLSLOp.h
+++ b/Engine/Plugins/Experimental/NNX/Source/NNXRuntimeRDG/Private/NNXRuntimeHLSLOp.h
@@ -13,7 +13,7 @@ namespace NNX
 	struct FMLOperatorHlsl : public FMLOperatorRDG
 	{
 		virtual ~FMLOperatorHlsl() = default;
-		virtual bool Initialize(TArrayView<const FTensor> InputTensors, TArrayView<const FTensor> OutputTensors, const UE::NNECore::FAttributeMap& Attributes) = 0;
+		virtual bool Initialize(TConstArrayView<FTensorDesc> InputTensorDescs, TConstArrayView<FTensorDesc> OutputTensorDescs, const UE::NNECore::FAttributeMap& Attributes) = 0;
 		virtual void Dispatch(FRDGBuilder& GraphBuilder, TArrayView<const FMLTensorBinding> InInputBindings, TArrayView<const FMLTensorBinding> OutOutputBindings) = 0;
 		virtual int ComputeOutputShape(TConstArrayView<FTensorShape> InputShapes, TArray<FTensorShape>& OutputShapes) const { return -1; };
 	};
