Add DirectML cross-platform support

#rnx

diff --git a/Engine/Plugins/Experimental/NNX/Source/NNXRuntimeRDG/NNXRuntimeRDG.Build.cs b/Engine/Plugins/Experimental/NNX/Source/NNXRuntimeRDG/NNXRuntimeRDG.Build.cs
--- a/Engine/Plugins/Experimental/NNX/Source/NNXRuntimeRDG/NNXRuntimeRDG.Build.cs
+++ b/Engine/Plugins/Experimental/NNX/Source/NNXRuntimeRDG/NNXRuntimeRDG.Build.cs
@@ -32,6 +32,8 @@ public class NNXRuntimeRDG : ModuleRules
 
 			AddEngineThirdPartyPrivateStaticDependencies(Target, "DX12");
 			AddEngineThirdPartyPrivateStaticDependencies(Target, "DirectMLDefault");
+
+			PublicDefinitions.Add("NNE_USE_DIRECTML");
 		}
 
 		if (Target.Platform == UnrealTargetPlatform.Mac)
diff --git a/Engine/Plugins/Experimental/NNX/Source/NNXRuntimeRDG/Private/NNXRuntimeDML.cpp b/Engine/Plugins/Experimental/NNX/Source/NNXRuntimeRDG/Private/NNXRuntimeDML.cpp
--- a/Engine/Plugins/Experimental/NNX/Source/NNXRuntimeRDG/Private/NNXRuntimeDML.cpp
+++ b/Engine/Plugins/Experimental/NNX/Source/NNXRuntimeRDG/Private/NNXRuntimeDML.cpp
@@ -13,25 +13,31 @@
 #include "RenderGraphBuilder.h"
 #include "RenderGraphUtils.h"
 
-// TODO: Use NNXOperator.h instead of HlslIdentityOp
 #include "NNXShaderParameters.h"
 
 #include "HAL/FileManager.h"
 
 // NOTE: For now we only have DML on Windows, we should add support for XSX
-#if PLATFORM_WINDOWS
+#ifdef NNE_USE_DIRECTML
 
 #include "Windows/AllowWindowsPlatformTypes.h"
+THIRD_PARTY_INCLUDES_START
 #include <unknwn.h>
 #include "Microsoft/COMPointer.h"
+#include "DirectML.h"
+THIRD_PARTY_INCLUDES_END
 #include "Windows/HideWindowsPlatformTypes.h"
 
+// DirectML is implemented using COM on all platforms
+#ifdef IID_GRAPHICS_PPV_ARGS
+#define DML_PPV_ARGS(x) __uuidof(*x), IID_PPV_ARGS_Helper(x)
+#else
+#define DML_PPV_ARGS(x) IID_PPV_ARGS(x)
+#endif
+
 #include "ID3D12DynamicRHI.h"
-#include "D3D12RHIBridge.h"
 //#include "D3D12RHI/Private/D3D12Device.h"	// It could be great if we could get the FD3D12DescriptorHeapManager instance
 
-#include "DirectML.h"
-
 namespace NNX
 {
 
@@ -431,7 +437,7 @@ protected:
 
 		HRESULT Res;
 		
-		Res = Device->CreateOperator(&DmlOpDesc, IID_PPV_ARGS(&DmlOp));
+		Res = Device->CreateOperator(&DmlOpDesc, DML_PPV_ARGS(&DmlOp));
 		if (!DmlOp)
 		{
 			UE_LOG(LogNNX, Warning, TEXT("Failed to creat DML operator, hres:%d"), Res);
@@ -439,7 +445,7 @@ protected:
 		}
 
 		// Compile operator
-		Res = Device->CompileOperator(DmlOp, DML_EXECUTION_FLAG_NONE, IID_PPV_ARGS(&CompiledOp));
+		Res = Device->CompileOperator(DmlOp, DML_EXECUTION_FLAG_NONE, DML_PPV_ARGS(&CompiledOp));
 		if (!CompiledOp)
 		{
 			UE_LOG(LogNNX, Warning, TEXT("Failed to compile DML operator"));
@@ -450,7 +456,7 @@ protected:
 		IDMLCompiledOperator* DmlOps[] = { CompiledOp };
 		TComPtr<IDMLOperatorInitializer>	DmlOpInit;
 
-		Res = Device->CreateOperatorInitializer(1, DmlOps, IID_PPV_ARGS(&DmlOpInit));
+		Res = Device->CreateOperatorInitializer(1, DmlOps, DML_PPV_ARGS(&DmlOpInit));
 		if (!DmlOpInit)
 		{
 			UE_LOG(LogNNX, Warning, TEXT("Failed to create DML operator initializer"));
@@ -477,7 +483,7 @@ protected:
 		DmlBindingTableDesc.GPUDescriptorHandle = CD3DX12_GPU_DESCRIPTOR_HANDLE(DevCtx->DescHeap->GetGPUDescriptorHandleForHeapStart(), DevCtx->NumDescriptors, DevCtx->DescriptorSize);
 		DmlBindingTableDesc.SizeInDescriptors = NumRequiredDescriptors;
 
-		Res = DevCtx->Device->CreateBindingTable(&DmlBindingTableDesc, IID_PPV_ARGS(&BindingTable));
+		Res = DevCtx->Device->CreateBindingTable(&DmlBindingTableDesc, DML_PPV_ARGS(&BindingTable));
 		if (!BindingTable)
 		{
 			UE_LOG(LogNNX, Warning, TEXT("Failed to create DML binding table, res:%d"), Res);
@@ -1289,6 +1295,7 @@ bool FMLRuntimeDml::Init()
 	Ctx.DeviceIndex = 0;
 	Ctx.D3D12Device = RHI->RHIGetDevice(Ctx.DeviceIndex);
 
+#if PLATFORM_WINDOWS
 	ID3D12Device5* D3D12Device5 = nullptr;
 	
 	Res = Ctx.D3D12Device->QueryInterface(&D3D12Device5);
@@ -1299,7 +1306,7 @@ bool FMLRuntimeDml::Init()
 		Res = D3D12Device5->EnumerateMetaCommands(&NumCommands, nullptr);
 		if (NumCommands)
 		{
-			UE_LOG(LogNNX, Display, TEXT("D3D12 Meta commands:%u"), NumCommands);
+			UE_LOG(LogNNX, Verbose, TEXT("D3D12 Meta commands:%u"), NumCommands);
 
 			TArray<D3D12_META_COMMAND_DESC>	MetaCmds;
 
@@ -1310,10 +1317,11 @@ bool FMLRuntimeDml::Init()
 			{
 				const D3D12_META_COMMAND_DESC& Desc = MetaCmds[Idx];
 
-				UE_LOG(LogNNX, Display, TEXT("   %s"), Desc.Name);
+				UE_LOG(LogNNX, Verbose, TEXT("   %s"), Desc.Name);
 			}
 		}
 	}
+#endif
 
 	DML_CREATE_DEVICE_FLAGS DmlCreateFlags = DML_CREATE_DEVICE_FLAG_NONE;
 
@@ -1323,14 +1331,14 @@ bool FMLRuntimeDml::Init()
 		DmlCreateFlags |= DML_CREATE_DEVICE_FLAG_DEBUG;
 	}
 
-	Res = DMLCreateDevice(Ctx.D3D12Device, DmlCreateFlags, IID_PPV_ARGS(&Ctx.Device));
+	Res = DMLCreateDevice(Ctx.D3D12Device, DmlCreateFlags, DML_PPV_ARGS(&Ctx.Device));
 	if (!Ctx.Device)
 	{
 		UE_LOG(LogNNX, Warning, TEXT("Failed to create DML device, res:%x"), Res);
 		return false;
 	}
 	
-	Res = Ctx.Device->CreateCommandRecorder(IID_PPV_ARGS(&Ctx.CmdRec));
+	Res = Ctx.Device->CreateCommandRecorder(DML_PPV_ARGS(&Ctx.CmdRec));
 	if (!Ctx.CmdRec)
 	{
 		UE_LOG(LogNNX, Warning, TEXT("Failed to create DML command recorder, res:%x"), Res);
@@ -1634,6 +1642,8 @@ IRuntime* FMLRuntimeDmlStartup()
 		//const FString GameDir = FModuleManager::Get().GetGameBinariesDirectory();
 		//const FString PluginDir = FPlatformProcess::BaseDir();
 
+#ifdef DIRECTML_BIN_PATH
+
 		//const FString PluginDir = IPluginManager::Get().FindPlugin(TEXT("NeuralNetworkInference"))->GetBaseDir();
 		const FString DirectMLRuntimeBinPath = TEXT(PREPROCESSOR_TO_STRING(DIRECTML_BIN_PATH));
 		//const FString DirectMLDLLPath = DirectMLRuntimeBinPath / TEXT(PREPROCESSOR_TO_STRING(DIRECTML_DLL_NAME));
@@ -1663,7 +1673,7 @@ IRuntime* FMLRuntimeDmlStartup()
 		}
 
 		FPlatformProcess::PopDllDirectory(*DirectMLRuntimeBinPath);
-
+#endif
 
 //#if DEBUG
 		//FPlatformProcess::GetDllHandle(TEXT("DirectML.Debug.dll"));
diff --git a/Engine/Plugins/Experimental/NNX/Source/NNXRuntimeRDG/Private/NNXRuntimeRDG.h b/Engine/Plugins/Experimental/NNX/Source/NNXRuntimeRDG/Private/NNXRuntimeRDG.h
--- a/Engine/Plugins/Experimental/NNX/Source/NNXRuntimeRDG/Private/NNXRuntimeRDG.h
+++ b/Engine/Plugins/Experimental/NNX/Source/NNXRuntimeRDG/Private/NNXRuntimeRDG.h
@@ -277,7 +277,7 @@ public:
 };
 
 // NOTE: For now we only have DML on Windows, we should add support for XSX
-#if PLATFORM_WINDOWS
+#ifdef NNE_USE_DIRECTML
 
 extern IRuntime* FMLRuntimeDmlStartup();
 extern void FMLRuntimeDmlShutdown();
diff --git a/Engine/Plugins/Experimental/NNX/Source/NNXRuntimeRDG/Private/NNXRuntimeRDGModule.cpp b/Engine/Plugins/Experimental/NNX/Source/NNXRuntimeRDG/Private/NNXRuntimeRDGModule.cpp
--- a/Engine/Plugins/Experimental/NNX/Source/NNXRuntimeRDG/Private/NNXRuntimeRDGModule.cpp
+++ b/Engine/Plugins/Experimental/NNX/Source/NNXRuntimeRDG/Private/NNXRuntimeRDGModule.cpp
@@ -19,7 +19,7 @@ public:
 
 	virtual void StartupModule() override
 	{
-#if PLATFORM_WINDOWS
+#ifdef NNE_USE_DIRECTML
 		DmlRuntime = NNX::FMLRuntimeDmlStartup();
 
 		if (DmlRuntime)
@@ -38,7 +38,7 @@ public:
 
 	virtual void ShutdownModule() override
 	{
-#if PLATFORM_WINDOWS
+#ifdef NNE_USE_DIRECTML
 		if (DmlRuntime)
 		{
 			NNX::UnregisterRuntime(DmlRuntime);
