NNI RDG: feed shape to operator.Dispatch()

Problem:
* Operators need to get the shape at runtime as they might not be known a model initialisation time (for dynamic shapes)

Solution:
* Make RDG work internally on FTensorRDG instead of FMLTensorBinding, themselve a specialized version of Tensor containing a FRDGBufferRef

Tests:
* All NNX tests pass under win64 editor

diff --git a/Engine/Plugins/Experimental/NNX/Source/NNXCore/Public/NNXTypes.h b/Engine/Plugins/Experimental/NNX/Source/NNXCore/Public/NNXTypes.h
--- a/Engine/Plugins/Experimental/NNX/Source/NNXCore/Public/NNXTypes.h
+++ b/Engine/Plugins/Experimental/NNX/Source/NNXCore/Public/NNXTypes.h
@@ -211,6 +211,7 @@ public:
 /** Concrete tensor backed by data */
 class FTensor : public FTensorDescBase
 {
+protected:
 	FTensorShape Shape;
 	uint32 Volume = 0;
 	uint64 DataSize = 0;
diff --git a/Engine/Plugins/Experimental/NNX/Source/NNXRuntimeRDG/Private/Hlsl/NNIRuntimeRDGConv.cpp b/Engine/Plugins/Experimental/NNX/Source/NNXRuntimeRDG/Private/Hlsl/NNIRuntimeRDGConv.cpp
--- a/Engine/Plugins/Experimental/NNX/Source/NNXRuntimeRDG/Private/Hlsl/NNIRuntimeRDGConv.cpp
+++ b/Engine/Plugins/Experimental/NNX/Source/NNXRuntimeRDG/Private/Hlsl/NNIRuntimeRDGConv.cpp
@@ -81,7 +81,7 @@ namespace UE::NNIRuntimeRDG::Private::Hlsl
 			return true;
 		}
 
-		virtual void Dispatch(FRDGBuilder& GraphBuilder, TArrayView<const NNX::FMLTensorBinding> InInputBindings, TArrayView<const NNX::FMLTensorBinding> OutOutputBindings) override
+		virtual void Dispatch(FRDGBuilder& GraphBuilder, TConstArrayView<NNX::FTensorRDG> InInputTensors, TConstArrayView<NNX::FTensorRDG> InOutputTensors) override
 		{
 			using namespace UE::NNEHlslShaders::Internal;
 
@@ -93,12 +93,12 @@ namespace UE::NNIRuntimeRDG::Private::Hlsl
 			// Set parameters
 			FConvCS::FParameters* Params = GraphBuilder.AllocParameters<FConvCS::FParameters>();
 			FConvCS::FillInParameters(GroupSize, Input.GetShape().Data, Weights.GetShape().Data, HasBias, AutoPad, Group, Dilations,Strides, Pads, *Params);
-			Params->X = GraphBuilder.CreateSRV(FRDGBufferSRVDesc(InInputBindings[0].Buffer, PF_R32_FLOAT));
-			Params->W = GraphBuilder.CreateSRV(FRDGBufferSRVDesc(InInputBindings[1].Buffer, PF_R32_FLOAT));
-			if (InInputBindings.Num() == 3) {
-				Params->B = GraphBuilder.CreateSRV(FRDGBufferSRVDesc(InInputBindings[2].Buffer, PF_R32_FLOAT));
+			Params->X = GraphBuilder.CreateSRV(FRDGBufferSRVDesc(InInputTensors[0].GetBuffer(), PF_R32_FLOAT));
+			Params->W = GraphBuilder.CreateSRV(FRDGBufferSRVDesc(InInputTensors[1].GetBuffer(), PF_R32_FLOAT));
+			if (InInputTensors.Num() == 3) {
+				Params->B = GraphBuilder.CreateSRV(FRDGBufferSRVDesc(InInputTensors[2].GetBuffer(), PF_R32_FLOAT));
 			}
-			Params->Y = GraphBuilder.CreateUAV(FRDGBufferUAVDesc(OutOutputBindings[0].Buffer, PF_R32_FLOAT));
+			Params->Y = GraphBuilder.CreateUAV(FRDGBufferUAVDesc(InOutputTensors[0].GetBuffer(), PF_R32_FLOAT));
 
 			FConvCS::FPermutationDomain PermutationVector;
 
diff --git a/Engine/Plugins/Experimental/NNX/Source/NNXRuntimeRDG/Private/Hlsl/NNIRuntimeRDGConvTranspose.cpp b/Engine/Plugins/Experimental/NNX/Source/NNXRuntimeRDG/Private/Hlsl/NNIRuntimeRDGConvTranspose.cpp
--- a/Engine/Plugins/Experimental/NNX/Source/NNXRuntimeRDG/Private/Hlsl/NNIRuntimeRDGConvTranspose.cpp
+++ b/Engine/Plugins/Experimental/NNX/Source/NNXRuntimeRDG/Private/Hlsl/NNIRuntimeRDGConvTranspose.cpp
@@ -82,7 +82,7 @@ namespace UE::NNIRuntimeRDG::Private::Hlsl
 			return true;
 		}
 
-		virtual void Dispatch(FRDGBuilder& GraphBuilder, TArrayView<const NNX::FMLTensorBinding> InInputBindings, TArrayView<const NNX::FMLTensorBinding> OutOutputBindings) override
+		virtual void Dispatch(FRDGBuilder& GraphBuilder, TConstArrayView<NNX::FTensorRDG> InInputTensors, TConstArrayView<NNX::FTensorRDG> InOutputTensors) override
 		{
 			using namespace UE::NNEHlslShaders::Internal;
 
@@ -94,12 +94,12 @@ namespace UE::NNIRuntimeRDG::Private::Hlsl
 			// Set parameters
 			FConvTransposeCS::FParameters* Params = GraphBuilder.AllocParameters<FConvTransposeCS::FParameters>();
 			FConvTransposeCS::FillInParameters(GroupSize, Input.GetShape().Data, Weights.GetShape().Data, HasBias, AutoPad, Group, Dilations,Strides, Pads, OutputPadding, *Params);
-			Params->X = GraphBuilder.CreateSRV(FRDGBufferSRVDesc(InInputBindings[0].Buffer, PF_R32_FLOAT));
-			Params->W = GraphBuilder.CreateSRV(FRDGBufferSRVDesc(InInputBindings[1].Buffer, PF_R32_FLOAT));
-			if (InInputBindings.Num() == 3) {
-				Params->B = GraphBuilder.CreateSRV(FRDGBufferSRVDesc(InInputBindings[2].Buffer, PF_R32_FLOAT));
+			Params->X = GraphBuilder.CreateSRV(FRDGBufferSRVDesc(InInputTensors[0].GetBuffer(), PF_R32_FLOAT));
+			Params->W = GraphBuilder.CreateSRV(FRDGBufferSRVDesc(InInputTensors[1].GetBuffer(), PF_R32_FLOAT));
+			if (InInputTensors.Num() == 3) {
+				Params->B = GraphBuilder.CreateSRV(FRDGBufferSRVDesc(InInputTensors[2].GetBuffer(), PF_R32_FLOAT));
 			}
-			Params->Y = GraphBuilder.CreateUAV(FRDGBufferUAVDesc(OutOutputBindings[0].Buffer, PF_R32_FLOAT));
+			Params->Y = GraphBuilder.CreateUAV(FRDGBufferUAVDesc(InOutputTensors[0].GetBuffer(), PF_R32_FLOAT));
 
 			FConvTransposeCS::FPermutationDomain PermutationVector;
 
diff --git a/Engine/Plugins/Experimental/NNX/Source/NNXRuntimeRDG/Private/Hlsl/NNIRuntimeRDGElementWiseBinary.cpp b/Engine/Plugins/Experimental/NNX/Source/NNXRuntimeRDG/Private/Hlsl/NNIRuntimeRDGElementWiseBinary.cpp
--- a/Engine/Plugins/Experimental/NNX/Source/NNXRuntimeRDG/Private/Hlsl/NNIRuntimeRDGElementWiseBinary.cpp
+++ b/Engine/Plugins/Experimental/NNX/Source/NNXRuntimeRDG/Private/Hlsl/NNIRuntimeRDGElementWiseBinary.cpp
@@ -51,12 +51,11 @@ namespace UE::NNIRuntimeRDG::Private::Hlsl
 			return true;
 		}
 
-		virtual void Dispatch(FRDGBuilder& GraphBuilder, TArrayView<const NNX::FMLTensorBinding> InInputBindings, TArrayView<const NNX::FMLTensorBinding> OutOutputBindings) override
+		virtual void Dispatch(FRDGBuilder& GraphBuilder, TConstArrayView<NNX::FTensorRDG> InInputTensors, TConstArrayView<NNX::FTensorRDG> InOutputTensors) override
 		{
-			// HACK: This only works for single layer networks
-			FRDGBufferSRVRef LHSInputSRV = GraphBuilder.CreateSRV(FRDGBufferSRVDesc(InInputBindings[0].Buffer, PF_R32_FLOAT));
-			FRDGBufferSRVRef RHSInputSRV = GraphBuilder.CreateSRV(FRDGBufferSRVDesc(InInputBindings[1].Buffer, PF_R32_FLOAT));
-			FRDGBufferUAVRef OutputUAV = GraphBuilder.CreateUAV(FRDGBufferUAVDesc(OutOutputBindings[0].Buffer, PF_R32_FLOAT));
+			FRDGBufferSRVRef LHSInputSRV = GraphBuilder.CreateSRV(FRDGBufferSRVDesc(InInputTensors[0].GetBuffer(), PF_R32_FLOAT));
+			FRDGBufferSRVRef RHSInputSRV = GraphBuilder.CreateSRV(FRDGBufferSRVDesc(InInputTensors[1].GetBuffer(), PF_R32_FLOAT));
+			FRDGBufferUAVRef OutputUAV = GraphBuilder.CreateUAV(FRDGBufferUAVDesc(InOutputTensors[0].GetBuffer(), PF_R32_FLOAT));
 
 			FIntVector ThreadGroupCount = NNX::ComputeElementWiseThreadGroups(Output.GetVolume(), FElementWiseBinaryConstants::NUM_GROUP_THREADS);
 
diff --git a/Engine/Plugins/Experimental/NNX/Source/NNXRuntimeRDG/Private/Hlsl/NNIRuntimeRDGElementWiseUnary.cpp b/Engine/Plugins/Experimental/NNX/Source/NNXRuntimeRDG/Private/Hlsl/NNIRuntimeRDGElementWiseUnary.cpp
--- a/Engine/Plugins/Experimental/NNX/Source/NNXRuntimeRDG/Private/Hlsl/NNIRuntimeRDGElementWiseUnary.cpp
+++ b/Engine/Plugins/Experimental/NNX/Source/NNXRuntimeRDG/Private/Hlsl/NNIRuntimeRDGElementWiseUnary.cpp
@@ -28,7 +28,6 @@ namespace UE::NNIRuntimeRDG::Private::Hlsl
 		float Alpha = 0.0f;
 		float Beta = 0.0f;
 		float Gamma = 0.0f;
-		int32 ElemByteSize = 0;
 
 	public:
 
@@ -47,41 +46,15 @@ namespace UE::NNIRuntimeRDG::Private::Hlsl
 			Alpha = Attributes.GetValueOrDefault(TEXT("alpha"), Alpha);
 			Beta = Attributes.GetValueOrDefault(TEXT("beta"), Beta);
 			Gamma = Attributes.GetValueOrDefault(TEXT("gamma"), Gamma);
-			ElemByteSize = InputTensorDescs[0].GetElemByteSize();
 			return true;
 		}
 
-		virtual void Dispatch(FRDGBuilder& GraphBuilder, TArrayView<const NNX::FMLTensorBinding> InInputBindings, TArrayView<const NNX::FMLTensorBinding> OutOutputBindings) override
+		virtual void Dispatch(FRDGBuilder& GraphBuilder, TConstArrayView<NNX::FTensorRDG> InInputTensors, TConstArrayView<NNX::FTensorRDG> InOutputTensors) override
 		{
-            //struct RDGTensor
-            //{
-            //    RDGBuffer*
-            //    TensorShape ? Do we need the type? aka TensorDesc?
-            //}
-
-			//Dispatch(FRDGBuilder & GraphBuilder, 
-			         //TConstArrayView<NNX::RDGTensor> InputBuffers,
-					 //TConstArrayView<NNX::RDGTensor> OutputBuffers) override
-
-            //or 
-            //Dispatch(FRDGBuilder & GraphBuilder, 
-                     //TConstArrayView<NNX::FConcreteShape> InputShapes,
-			         //TConstArrayView<NNX::RDGTensor*> InputBuffers,
-                     //TConstArrayView<NNX::FConcreteShape> OutputShapes,
-					 //TConstArrayView<NNX::RDGTensor*> OutputBuffers) override
-
-            //or 
-            //Dispatch(FRDGBuilder & GraphBuilder, 
-                     //TConstArrayView<NNX::RDGTensor*> Buffers,
-                     //TConstArrayView<NNX::FConcreteShape> Shapes,
-			         //TConstArrayView<uint32> InputIndices,
-                     //TConstArrayView<uint32> OutputIndices) override
-
-			// HACK: This only works for single layer networks
-			FRDGBufferSRVRef InputSRV = GraphBuilder.CreateSRV(FRDGBufferSRVDesc(InInputBindings[0].Buffer, PF_R32_FLOAT));
-			FRDGBufferUAVRef OutputUAV = GraphBuilder.CreateUAV(FRDGBufferUAVDesc(OutOutputBindings[0].Buffer, PF_R32_FLOAT));
+			FRDGBufferSRVRef InputSRV = GraphBuilder.CreateSRV(FRDGBufferSRVDesc(InInputTensors[0].GetBuffer(), PF_R32_FLOAT));
+			FRDGBufferUAVRef OutputUAV = GraphBuilder.CreateUAV(FRDGBufferUAVDesc(InOutputTensors[0].GetBuffer(), PF_R32_FLOAT));
 		
-			int32 NumElements = OutOutputBindings[0].SizeInBytes / ElemByteSize;//TODO jira 169354: Use shape from Dispatch instead
+			int32 NumElements = InOutputTensors[0].GetVolume();
 			FIntVector ThreadGroupCount = NNX::ComputeElementWiseThreadGroups(NumElements, FElementWiseUnaryConstants::NUM_GROUP_THREADS);
 
 			// Set parameters
diff --git a/Engine/Plugins/Experimental/NNX/Source/NNXRuntimeRDG/Private/Hlsl/NNIRuntimeRDGElementWiseVariadic.cpp b/Engine/Plugins/Experimental/NNX/Source/NNXRuntimeRDG/Private/Hlsl/NNIRuntimeRDGElementWiseVariadic.cpp
--- a/Engine/Plugins/Experimental/NNX/Source/NNXRuntimeRDG/Private/Hlsl/NNIRuntimeRDGElementWiseVariadic.cpp
+++ b/Engine/Plugins/Experimental/NNX/Source/NNXRuntimeRDG/Private/Hlsl/NNIRuntimeRDGElementWiseVariadic.cpp
@@ -12,30 +12,27 @@ namespace UE::NNIRuntimeRDG::Private::Hlsl
 	using FElementWiseVariadicConstants = UE::NNEHlslShaders::Internal::FElementWiseVariadicConstants;
 
 	void AddOneVariadicOpPass(FRDGBuilder& GraphBuilder, 
-		TArrayView<const NNX::FMLTensorBinding> InputBindings,
-		TArrayView<const NNX::FTensor> InputDesc,
-		const NNX::FMLTensorBinding& OutputBinding,
-		const NNX::FTensor& OutputDesc,
+		TConstArrayView<NNX::FTensorRDG> InputTensors,
+		const NNX::FTensorRDG& OutputTensor,
 		bool OutputAsInput,
 		EMLElementWiseVariadicOperatorType OpType,
 		float Scale)
 	{
 		static_assert(FElementWiseVariadicConstants::MAX_NUM_INPUT == 4, "This algorithm need to be adapted to math the shader.");
-		check(InputBindings.Num() == InputDesc.Num());
-		check(InputBindings.Num() <= FElementWiseVariadicConstants::MAX_NUM_INPUT);
-		check(InputBindings.Num() > 0);
+		check(InputTensors.Num() <= FElementWiseVariadicConstants::MAX_NUM_INPUT);
+		check(InputTensors.Num() > 0);
 
 		// SRVs & UAV creations
-		FRDGBufferUAVRef OutputUAV = GraphBuilder.CreateUAV(FRDGBufferUAVDesc(OutputBinding.Buffer, PF_R32_FLOAT));
+		FRDGBufferUAVRef OutputUAV = GraphBuilder.CreateUAV(FRDGBufferUAVDesc(OutputTensor.GetBuffer(), PF_R32_FLOAT));
 		FRDGBufferSRVRef InputsSRV[FElementWiseVariadicConstants::MAX_NUM_INPUT] = { nullptr, nullptr, nullptr, nullptr };
 
-		for (int32 i = 0; i < InputBindings.Num(); ++i)
+		for (int32 i = 0; i < InputTensors.Num(); ++i)
 		{
-			InputsSRV[i] = GraphBuilder.CreateSRV(FRDGBufferSRVDesc(InputBindings[i].Buffer, PF_R32_FLOAT));
+			InputsSRV[i] = GraphBuilder.CreateSRV(FRDGBufferSRVDesc(InputTensors[i].GetBuffer(), PF_R32_FLOAT));
 		}
 
 		// Set parameters
-		FIntVector ThreadGroupCount = NNX::ComputeElementWiseThreadGroups(OutputDesc.GetVolume(), FElementWiseVariadicConstants::NUM_GROUP_THREADS);
+		FIntVector ThreadGroupCount = NNX::ComputeElementWiseThreadGroups(OutputTensor.GetVolume(), FElementWiseVariadicConstants::NUM_GROUP_THREADS);
 		TElementWiseVariadicCS::FParameters* Params = GraphBuilder.AllocParameters<TElementWiseVariadicCS::FParameters>();
 
 		Params->Input0 = InputsSRV[0];
@@ -43,21 +40,21 @@ namespace UE::NNIRuntimeRDG::Private::Hlsl
 		Params->Input2 = InputsSRV[2];
 		Params->Input3 = InputsSRV[3];
 		Params->Output = OutputUAV;
-		FillTensorStrideForBroadcastShaderParameters(InputDesc[0], OutputDesc.GetShape().Rank(), Params->InputTensorInfo, 0);
-		if (InputBindings.Num() >= 2)
+		FillTensorStrideForBroadcastShaderParameters(InputTensors[0], OutputTensor.GetShape().Rank(), Params->InputTensorInfo, 0);
+		if (InputTensors.Num() >= 2)
 		{
-			FillTensorStrideForBroadcastShaderParameters(InputDesc[1], OutputDesc.GetShape().Rank(), Params->InputTensorInfo, 1);
+			FillTensorStrideForBroadcastShaderParameters(InputTensors[1], OutputTensor.GetShape().Rank(), Params->InputTensorInfo, 1);
 		}
-		if (InputBindings.Num() >= 3)
+		if (InputTensors.Num() >= 3)
 		{
-			FillTensorStrideForBroadcastShaderParameters(InputDesc[2], OutputDesc.GetShape().Rank(), Params->InputTensorInfo, 2);
+			FillTensorStrideForBroadcastShaderParameters(InputTensors[2], OutputTensor.GetShape().Rank(), Params->InputTensorInfo, 2);
 		}
-		if (InputBindings.Num() >= 4)
+		if (InputTensors.Num() >= 4)
 		{
-			FillTensorStrideForBroadcastShaderParameters(InputDesc[3], OutputDesc.GetShape().Rank(), Params->InputTensorInfo, 3);
+			FillTensorStrideForBroadcastShaderParameters(InputTensors[3], OutputTensor.GetShape().Rank(), Params->InputTensorInfo, 3);
 		}
-		FillTensorStrideShaderParameters(OutputDesc, Params->OutputTensorInfo, 0);
-		Params->Num = OutputDesc.GetVolume();
+		FillTensorStrideShaderParameters(OutputTensor, Params->OutputTensorInfo, 0);
+		Params->Num = OutputTensor.GetVolume();
 		Params->ThreadCountX = ThreadGroupCount.X * FElementWiseVariadicConstants::NUM_GROUP_THREADS;
 		Params->Scale = Scale;
 
@@ -67,8 +64,8 @@ namespace UE::NNIRuntimeRDG::Private::Hlsl
 		PermutationVector.Set<TElementWiseVariadicCS::FOperatorType>(OpType);
 		PermutationVector.Set<TElementWiseVariadicCS::FApplyScale>(Scale != 1.0f);
 		PermutationVector.Set<TElementWiseVariadicCS::FOutputAsInput>(OutputAsInput);
-		PermutationVector.Set<TElementWiseVariadicCS::FNumInput>(InputBindings.Num());
-		PermutationVector.Set<TElementWiseVariadicCS::FVariadicNumDimensions>(OutputDesc.GetShape().Rank());
+		PermutationVector.Set<TElementWiseVariadicCS::FNumInput>(InputTensors.Num());
+		PermutationVector.Set<TElementWiseVariadicCS::FVariadicNumDimensions>(OutputTensor.GetShape().Rank());
 
 		// Add the pass to RDG
 		TShaderMapRef<TElementWiseVariadicCS> ComputeShader(GetGlobalShaderMap(GMaxRHIFeatureLevel), PermutationVector);
@@ -92,46 +89,27 @@ namespace UE::NNIRuntimeRDG::Private::Hlsl
 		TElementWiseVariadic() {}
 		virtual ~TElementWiseVariadic() = default;
 
-	private:
-
-		TArray<NNX::FTensor> InputDescs;
-		NNX::FTensor OutputDesc = {};
-
 	public:
 
 		virtual bool Initialize(TConstArrayView<NNX::FTensorDesc> InputTensorDescs, TConstArrayView<NNX::FTensorDesc> OutputTensorDescs, const UE::NNECore::FAttributeMap& Attributes) override
 		{
-			TArray<NNX::FTensor> InputTensors;
-			TArray<NNX::FTensor> OutputTensors;
-			if (!NNX::ConvertConcreteTensorDescsToTensors(InputTensorDescs, InputTensors) ||
-				!NNX::ConvertConcreteTensorDescsToTensors(OutputTensorDescs, OutputTensors))
-			{
-				UE_LOG(LogNNX, Warning, TEXT("Variable input shapes are not supported by this operator"));
-				return false;
-			}
-			
-			check(InputTensors.Num() > 0);
-			check(OutputTensors.Num() == 1);
-		
-			InputDescs.Append(InputTensors);
-			OutputDesc = OutputTensors[0];
-
+			check(InputTensorDescs.Num() > 0);
+			check(OutputTensorDescs.Num() == 1);
 			return true;
 		}
 
-		virtual void Dispatch(FRDGBuilder& GraphBuilder, TArrayView<const NNX::FMLTensorBinding> InInputBindings, TArrayView<const NNX::FMLTensorBinding> OutOutputBindings) override
+		virtual void Dispatch(FRDGBuilder& GraphBuilder, TConstArrayView<NNX::FTensorRDG> InInputTensors, TConstArrayView<NNX::FTensorRDG> InOutputTensors) override
 		{
-			check(OutOutputBindings.Num() == 1);
-			check(InInputBindings.Num() == InputDescs.Num());
+			check(InInputTensors.Num() > 0);
+			check(InOutputTensors.Num() == 1);
 
 			RDG_EVENT_SCOPE(GraphBuilder, "NNI.Operator.Hlsl.ElementWise.Variadic");
 			RDG_GPU_STAT_SCOPE(GraphBuilder, FNNIOperatorElementWiseVariadic);
 
-			NNX::FMLTensorBinding PassInputBindings[FElementWiseVariadicConstants::MAX_NUM_INPUT];
-			NNX::FTensor PassInputDescs[FElementWiseVariadicConstants::MAX_NUM_INPUT];
-			for (int32 InputOffset = 0; InputOffset < InInputBindings.Num(); InputOffset += FElementWiseVariadicConstants::MAX_NUM_INPUT)
+			NNX::FTensorRDG PassInputTensors[FElementWiseVariadicConstants::MAX_NUM_INPUT];
+			for (int32 InputOffset = 0; InputOffset < InInputTensors.Num(); InputOffset += FElementWiseVariadicConstants::MAX_NUM_INPUT)
 			{
-				uint32 NumInputLeftToHandle = InInputBindings.Num() - InputOffset;
+				uint32 NumInputLeftToHandle = InInputTensors.Num() - InputOffset;
 				uint32 NumInputForPass = FMath::Min(FElementWiseVariadicConstants::MAX_NUM_INPUT, NumInputLeftToHandle);
 				bool bIsFirstPass = (InputOffset == 0);
 				bool bIsLastPass = (NumInputLeftToHandle <= FElementWiseVariadicConstants::MAX_NUM_INPUT);
@@ -139,20 +117,17 @@ namespace UE::NNIRuntimeRDG::Private::Hlsl
 
 				for (uint32 i = 0; i < NumInputForPass; ++i)
 				{
-					PassInputBindings[i] = InInputBindings[InputOffset + i];
-					PassInputDescs[i] = InputDescs[InputOffset + i];
+					PassInputTensors[i] = InInputTensors[InputOffset + i];
 				}
 
 				if (OpType == EMLElementWiseVariadicOperatorType::Mean && bIsLastPass)
 				{
-					Scale = 1.0f / InInputBindings.Num();
+					Scale = 1.0f / InInputTensors.Num();
 				}
 			
 				AddOneVariadicOpPass(GraphBuilder,
-					MakeArrayView(PassInputBindings, NumInputForPass),
-					MakeArrayView(PassInputDescs, NumInputForPass),
-					OutOutputBindings[0],
-					OutputDesc,
+					MakeArrayView(PassInputTensors, NumInputForPass),
+					InOutputTensors[0],
 					!bIsFirstPass,
 					OpType, Scale);
 			}
diff --git a/Engine/Plugins/Experimental/NNX/Source/NNXRuntimeRDG/Private/Hlsl/NNIRuntimeRDGGather.cpp b/Engine/Plugins/Experimental/NNX/Source/NNXRuntimeRDG/Private/Hlsl/NNIRuntimeRDGGather.cpp
--- a/Engine/Plugins/Experimental/NNX/Source/NNXRuntimeRDG/Private/Hlsl/NNIRuntimeRDGGather.cpp
+++ b/Engine/Plugins/Experimental/NNX/Source/NNXRuntimeRDG/Private/Hlsl/NNIRuntimeRDGGather.cpp
@@ -61,16 +61,16 @@ namespace UE::NNIRuntimeRDG::Private::Hlsl
 			return true;
 		}
 
-		virtual void Dispatch(FRDGBuilder& GraphBuilder, TArrayView<const NNX::FMLTensorBinding> InInputBindings, TArrayView<const NNX::FMLTensorBinding> OutOutputBindings) override
+		virtual void Dispatch(FRDGBuilder& GraphBuilder, TConstArrayView<NNX::FTensorRDG> InInputTensors, TConstArrayView<NNX::FTensorRDG> InOutputTensors) override
 		{
 			using namespace UE::NNEHlslShaders::Internal;
 
 			// Set parameters
 			TGatherCS::FParameters* Parameters = GraphBuilder.AllocParameters<TGatherCS::FParameters>();
 			TGatherCS::FillInParameters(Axis, Data, Indices, *Parameters);
-			Parameters->Data = GraphBuilder.CreateSRV(FRDGBufferSRVDesc(InInputBindings[0].Buffer, PF_R32_FLOAT));
-			Parameters->Indices = GraphBuilder.CreateSRV(FRDGBufferSRVDesc(InInputBindings[1].Buffer, PF_R32_FLOAT));
-			Parameters->Output = GraphBuilder.CreateUAV(FRDGBufferUAVDesc(OutOutputBindings[0].Buffer, PF_R32_FLOAT));
+			Parameters->Data = GraphBuilder.CreateSRV(FRDGBufferSRVDesc(InInputTensors[0].GetBuffer(), PF_R32_FLOAT));
+			Parameters->Indices = GraphBuilder.CreateSRV(FRDGBufferSRVDesc(InInputTensors[1].GetBuffer(), PF_R32_FLOAT));
+			Parameters->Output = GraphBuilder.CreateUAV(FRDGBufferUAVDesc(InOutputTensors[0].GetBuffer(), PF_R32_FLOAT));
 
 			TGatherCS::FPermutationDomain PermutationVector;
 			PermutationVector.Set<TGatherCS::FGatherNumOutputDimensions>(Output.GetShape().Rank());
diff --git a/Engine/Plugins/Experimental/NNX/Source/NNXRuntimeRDG/Private/Hlsl/NNIRuntimeRDGGemm.cpp b/Engine/Plugins/Experimental/NNX/Source/NNXRuntimeRDG/Private/Hlsl/NNIRuntimeRDGGemm.cpp
--- a/Engine/Plugins/Experimental/NNX/Source/NNXRuntimeRDG/Private/Hlsl/NNIRuntimeRDGGemm.cpp
+++ b/Engine/Plugins/Experimental/NNX/Source/NNXRuntimeRDG/Private/Hlsl/NNIRuntimeRDGGemm.cpp
@@ -79,7 +79,7 @@ namespace UE::NNIRuntimeRDG::Private::Hlsl
 			return true;
 		}
 
-		virtual void Dispatch(FRDGBuilder& GraphBuilder, TArrayView<const NNX::FMLTensorBinding> InInputBindings, TArrayView<const NNX::FMLTensorBinding> OutOutputBindings) override
+		virtual void Dispatch(FRDGBuilder& GraphBuilder, TConstArrayView<NNX::FTensorRDG> InInputTensors, TConstArrayView<NNX::FTensorRDG> InOutputTensors) override
 		{
 			using namespace UE::NNEHlslShaders::Internal;
 
@@ -90,12 +90,12 @@ namespace UE::NNIRuntimeRDG::Private::Hlsl
 			// Set parameters
 			TGemmCS::FParameters* Parameters = GraphBuilder.AllocParameters<TGemmCS::FParameters>();
 			TGemmCS::FillInParameters(InputAlpha, InputBeta, InputTransA, InputTransB, InputA, InputB, InputC, CScalar, *Parameters);
-			Parameters->A = GraphBuilder.CreateSRV(FRDGBufferSRVDesc(InInputBindings[0].Buffer, PF_R32_FLOAT));
-			Parameters->B = GraphBuilder.CreateSRV(FRDGBufferSRVDesc(InInputBindings[1].Buffer, PF_R32_FLOAT));
-			if (InInputBindings.Num() == 3) {
-				Parameters->C = GraphBuilder.CreateSRV(FRDGBufferSRVDesc(InInputBindings[2].Buffer, PF_R32_FLOAT));
+			Parameters->A = GraphBuilder.CreateSRV(FRDGBufferSRVDesc(InInputTensors[0].GetBuffer(), PF_R32_FLOAT));
+			Parameters->B = GraphBuilder.CreateSRV(FRDGBufferSRVDesc(InInputTensors[1].GetBuffer(), PF_R32_FLOAT));
+			if (InInputTensors.Num() == 3) {
+				Parameters->C = GraphBuilder.CreateSRV(FRDGBufferSRVDesc(InInputTensors[2].GetBuffer(), PF_R32_FLOAT));
 			}
-			Parameters->Y = GraphBuilder.CreateUAV(FRDGBufferUAVDesc(OutOutputBindings[0].Buffer, PF_R32_FLOAT));
+			Parameters->Y = GraphBuilder.CreateUAV(FRDGBufferUAVDesc(InOutputTensors[0].GetBuffer(), PF_R32_FLOAT));
 
 			TGemmCS::FPermutationDomain PermutationVector;
 			PermutationVector.Set<TGemmCS::FGemmCScalar>(bNoBias ? EGemmCScalar::NoBias : (bIsCScalar ? EGemmCScalar::Yes : EGemmCScalar::No));
diff --git a/Engine/Plugins/Experimental/NNX/Source/NNXRuntimeRDG/Private/Hlsl/NNIRuntimeRDGMatMul.cpp b/Engine/Plugins/Experimental/NNX/Source/NNXRuntimeRDG/Private/Hlsl/NNIRuntimeRDGMatMul.cpp
--- a/Engine/Plugins/Experimental/NNX/Source/NNXRuntimeRDG/Private/Hlsl/NNIRuntimeRDGMatMul.cpp
+++ b/Engine/Plugins/Experimental/NNX/Source/NNXRuntimeRDG/Private/Hlsl/NNIRuntimeRDGMatMul.cpp
@@ -47,7 +47,7 @@ namespace UE::NNIRuntimeRDG::Private::Hlsl
 			return true;
 		}
 
-		virtual void Dispatch(FRDGBuilder& GraphBuilder, TArrayView<const NNX::FMLTensorBinding> InInputBindings, TArrayView<const NNX::FMLTensorBinding> OutOutputBindings) override
+		virtual void Dispatch(FRDGBuilder& GraphBuilder, TConstArrayView<NNX::FTensorRDG> InInputTensors, TConstArrayView<NNX::FTensorRDG> InOutputTensors) override
 		{
 			using namespace UE::NNEHlslShaders::Internal;
 
@@ -58,9 +58,9 @@ namespace UE::NNIRuntimeRDG::Private::Hlsl
 			// Set parameters
 			TGemmCS::FParameters* Parameters = GraphBuilder.AllocParameters<TGemmCS::FParameters>();
 			TGemmCS::FillInParametersMatMul(InputA, InputB, *Parameters);
-			Parameters->A = GraphBuilder.CreateSRV(FRDGBufferSRVDesc(InInputBindings[0].Buffer, PF_R32_FLOAT));
-			Parameters->B = GraphBuilder.CreateSRV(FRDGBufferSRVDesc(InInputBindings[1].Buffer, PF_R32_FLOAT));
-			Parameters->Y = GraphBuilder.CreateUAV(FRDGBufferUAVDesc(OutOutputBindings[0].Buffer, PF_R32_FLOAT));
+			Parameters->A = GraphBuilder.CreateSRV(FRDGBufferSRVDesc(InInputTensors[0].GetBuffer(), PF_R32_FLOAT));
+			Parameters->B = GraphBuilder.CreateSRV(FRDGBufferSRVDesc(InInputTensors[1].GetBuffer(), PF_R32_FLOAT));
+			Parameters->Y = GraphBuilder.CreateUAV(FRDGBufferUAVDesc(InOutputTensors[0].GetBuffer(), PF_R32_FLOAT));
 
 			TGemmCS::FPermutationDomain PermutationVector;
 			PermutationVector.Set<TGemmCS::FGemmCScalar>(EGemmCScalar::NoBias);
diff --git a/Engine/Plugins/Experimental/NNX/Source/NNXRuntimeRDG/Private/NNXRuntimeDML.cpp b/Engine/Plugins/Experimental/NNX/Source/NNXRuntimeRDG/Private/NNXRuntimeDML.cpp
--- a/Engine/Plugins/Experimental/NNX/Source/NNXRuntimeRDG/Private/NNXRuntimeDML.cpp
+++ b/Engine/Plugins/Experimental/NNX/Source/NNXRuntimeRDG/Private/NNXRuntimeDML.cpp
@@ -334,8 +334,7 @@ class FMLOperatorDml : public FMLOperatorRDG
 public:
 
 	virtual bool Initialize(FDeviceContextDml* DevCtx, TArrayView<const FTensor> InputTensors, TArrayView<const FTensor> OutputTensors, const UE::NNECore::FAttributeMap& Attributes) = 0;
-
-	virtual void Dispatch(FRDGBuilder& GraphBuilder, TArrayView<const FMLTensorBinding> InInputBindings, TArrayView<const FMLTensorBinding> OutOutputBindings) = 0;
+	virtual void Dispatch(FRDGBuilder& GraphBuilder, TConstArrayView<FTensorRDG> InInputTensors, TConstArrayView<FTensorRDG> InOutputTensors) = 0;
 
 protected:
 
@@ -647,15 +646,14 @@ public:
 	//
 	//
 	//
-	virtual void Dispatch(FRDGBuilder& GraphBuilder, TArrayView<const FMLTensorBinding> InInputBindings, TArrayView<const FMLTensorBinding> OutOutputBindings) override
+	virtual void Dispatch(FRDGBuilder& GraphBuilder, TConstArrayView<FTensorRDG> InInputTensors, TConstArrayView<FTensorRDG> InOutputTensors) override
 	{
 		ID3D12DynamicRHI* DynamicRHI = GetID3D12DynamicRHI();
 
 		FMLElementWiseUnaryParameters* Params = GraphBuilder.AllocParameters<FMLElementWiseUnaryParameters>();
 
-		// HACK: This only works for single layer networks
-		Params->Input = InInputBindings[0].Buffer;
-		Params->Output = OutOutputBindings[0].Buffer;
+		Params->Input = InInputTensors[0].GetBuffer();
+		Params->Output = InOutputTensors[0].GetBuffer();
 		Params->Alpha = Alpha;
 		Params->Beta = Beta;
 		Params->Gamma = Gamma;
@@ -886,15 +884,15 @@ public:
 	//
 	//
 	//
-	virtual void Dispatch(FRDGBuilder& GraphBuilder, TArrayView<const FMLTensorBinding> InInputBindings, TArrayView<const FMLTensorBinding> OutOutputBindings) override
+	virtual void Dispatch(FRDGBuilder& GraphBuilder, TConstArrayView<FTensorRDG> InInputTensors, TConstArrayView<FTensorRDG> InOutputTensors) override
 	{
 		ID3D12DynamicRHI* DynamicRHI = GetID3D12DynamicRHI();
 
 		FMLElementWiseBinaryParameters* Params = GraphBuilder.AllocParameters<FMLElementWiseBinaryParameters>();
 
-		Params->LHSInput = InInputBindings[0].Buffer;
-		Params->RHSInput = InInputBindings[1].Buffer;
-		Params->Output = OutOutputBindings[0].Buffer;
+		Params->LHSInput = InInputTensors[0].GetBuffer();
+		Params->RHSInput = InInputTensors[1].GetBuffer();
+		Params->Output = InOutputTensors[0].GetBuffer();
 		Params->Num = Num;
 
 		GraphBuilder.AddPass(
@@ -1069,19 +1067,19 @@ public:
 	//
 	//
 	//
-	virtual void Dispatch(FRDGBuilder& GraphBuilder, TArrayView<const FMLTensorBinding> InInputBindings, TArrayView<const FMLTensorBinding> OutOutputBindings) override
+	virtual void Dispatch(FRDGBuilder& GraphBuilder, TConstArrayView<FTensorRDG> InInputTensors, TConstArrayView<FTensorRDG> InOutputTensors) override
 	{
 		ID3D12DynamicRHI* DynamicRHI = GetID3D12DynamicRHI();
 
-		const int32 NumInputs = InInputBindings.Num();
+		const int32 NumInputs = InInputTensors.Num();
 		const bool bIsUsingBias = NumInputs > 2;
 
 		FMLGemmParameters* Params = GraphBuilder.AllocParameters<FMLGemmParameters>();
 
-		Params->A = InInputBindings[0].Buffer;
-		Params->B = InInputBindings[1].Buffer;
-		Params->C = bIsUsingBias ? InInputBindings[2].Buffer : nullptr;
-		Params->Y = OutOutputBindings[0].Buffer;
+		Params->A = InInputTensors[0].GetBuffer();
+		Params->B = InInputTensors[1].GetBuffer();
+		Params->C = bIsUsingBias ? InInputTensors[2].GetBuffer() : nullptr;
+		Params->Y = InOutputTensors[0].GetBuffer();
 
 		GraphBuilder.AddPass(
 			RDG_EVENT_NAME("FMLGemmDml_Transition"),
@@ -1491,13 +1489,13 @@ bool FMLInferenceModelDml::Init(UMLInferenceModel* InModel, FDeviceContextDml* I
 
 		for (int32 InputTensorIndex : Format.Operators[Idx].InTensors)
 		{
-			FTensorDesc SymbolicTensorDesc = AllSymbolicTensors[InputTensorIndex];
+			FTensorDesc SymbolicTensorDesc = AllSymbolicTensorDescs[InputTensorIndex];
 			//TODO jira 168972: Handle dynamic tensor desc, op should init from symbolic shapes
 			OpInputTensors.Emplace(FTensor::MakeFromSymbolicDesc(SymbolicTensorDesc));
 		}
 		for (int32 OutputTensorIndex : Format.Operators[Idx].OutTensors)
 		{
-			FTensorDesc SymbolicTensorDesc = AllSymbolicTensors[OutputTensorIndex];
+			FTensorDesc SymbolicTensorDesc = AllSymbolicTensorDescs[OutputTensorIndex];
 			//TODO jira 168972: Handle dynamic tensor desc, op should init from symbolic shapes
 			OpOutputTensors.Emplace(FTensor::MakeFromSymbolicDesc(SymbolicTensorDesc));
 		}
@@ -1527,38 +1525,31 @@ bool FMLInferenceModelDml::Init(UMLInferenceModel* InModel, FDeviceContextDml* I
 //
 void FMLInferenceModelDml::AddDispatchOps_RenderThread(FRDGBuilder& GraphBuilder)
 {
-	check(AllTensorBindings.Num() == AllTensors.Num());
-	checkCode(
-		for (int32 i = 0; i < AllTensorBindings.Num(); ++i)
-		{
-			check(AllTensors[i].GetDataSize() <= AllTensorBindings[i].SizeInBytes);
-		}
-	);
+	check(AllTensorRDGs.Num() == AllShapes.Num());
 
 	static constexpr int32 MaxExpectedInput = 10;
-	TArray<FMLTensorBinding, TInlineAllocator<MaxExpectedInput>> InputBindings;
-	
+	TArray<FTensorRDG, TInlineAllocator<MaxExpectedInput>> InputTensors;
+
 	static constexpr int32 MaxExpectedOutput = 2;
-	TArray<FMLTensorBinding, TInlineAllocator<MaxExpectedOutput>> OutputBindings;
-	
+	TArray<FTensorRDG, TInlineAllocator<MaxExpectedOutput>> OutputTensors;
+
 	// Add passes for all operators
 	for (int32 Idx = 0; Idx < Operators.Num(); ++Idx)
 	{
-		InputBindings.Empty();
+		InputTensors.Empty();
 		for (int32 i : OperatorInputTensorIndices[Idx])
 		{
-			InputBindings.Emplace(AllTensorBindings[i]);
+			InputTensors.Emplace(AllTensorRDGs[i]);
 		}
-		OutputBindings.Empty();
+		OutputTensors.Empty();
 		for (int32 i : OperatorOutputTensorIndices[Idx])
 		{
-			OutputBindings.Emplace(AllTensorBindings[i]);
+			OutputTensors.Emplace(AllTensorRDGs[i]);
 		}
 
 		FMLOperatorDml* Op = Operators[Idx];
 
-		//TODO jira 169354 pass shape information to Op.Dispatch probably via an object containing both mem info and shape (from AllTensors).
-		Op->Dispatch(GraphBuilder, InputBindings, OutputBindings);
+		Op->Dispatch(GraphBuilder, InputTensors, OutputTensors);
 	}
 }
 
@@ -1595,19 +1586,19 @@ FMLOperatorDml* FMLInferenceModelDml::OpCreate(const FString& OpName, TArrayView
 //
 int FMLInferenceModelDml::RunShapeInference()
 {
-	AllTensors.Empty();
+	AllShapes.Empty();
 
-	for (FTensorDesc SymbolicTensorDesc : AllSymbolicTensors)
+	for (FTensorDesc SymbolicTensorDesc : AllSymbolicTensorDescs)
 	{
 		if (SymbolicTensorDesc.IsConcrete())
 		{
-			FTensor Tensor = FTensor::MakeFromSymbolicDesc(SymbolicTensorDesc);
-			AllTensors.Emplace(Tensor);
+			FTensorShape TensorShape = FTensorShape::MakeFromSymbolic(SymbolicTensorDesc.GetShape());
+			AllShapes.Emplace(TensorShape);
 		}
 	}
-	if (AllTensors.Num() != AllSymbolicTensors.Num())
+	if (AllShapes.Num() != AllSymbolicTensorDescs.Num())
 	{
-		AllTensors.Empty();
+		AllShapes.Empty();
 		UE_LOG(LogNNX, Warning, TEXT("DML engine does not support model with variable shapes yet."));
 		return -1;
 	}
diff --git a/Engine/Plugins/Experimental/NNX/Source/NNXRuntimeRDG/Private/NNXRuntimeHLSL.cpp b/Engine/Plugins/Experimental/NNX/Source/NNXRuntimeRDG/Private/NNXRuntimeHLSL.cpp
--- a/Engine/Plugins/Experimental/NNX/Source/NNXRuntimeRDG/Private/NNXRuntimeHLSL.cpp
+++ b/Engine/Plugins/Experimental/NNX/Source/NNXRuntimeRDG/Private/NNXRuntimeHLSL.cpp
@@ -84,11 +84,11 @@ bool FMLInferenceModelHlsl::Init(UMLInferenceModel* InModel)
 
 		for (int32 InputTensorIndex : Format.Operators[Idx].InTensors)
 		{
-			Inputs.Emplace(AllSymbolicTensors[InputTensorIndex]);
+			Inputs.Emplace(AllSymbolicTensorDescs[InputTensorIndex]);
 		}
 		for (int32 OutputTensorIndex : Format.Operators[Idx].OutTensors)
 		{
-			Outputs.Emplace(AllSymbolicTensors[OutputTensorIndex]);
+			Outputs.Emplace(AllSymbolicTensorDescs[OutputTensorIndex]);
 		}
 		for (const FMLFormatAttributeDesc& Desc : Format.Operators[Idx].Attributes)
 		{
@@ -138,7 +138,7 @@ FMLOperatorHlsl* FMLInferenceModelHlsl::OpCreate(const FString& OpName, TConstAr
 
 int FMLInferenceModelHlsl::RunShapeInference()
 {
-	AllTensors.Empty();
+	AllShapes.Empty();
 
 	if (Operators.Num() == 0)
 	{
@@ -149,11 +149,10 @@ int FMLInferenceModelHlsl::RunShapeInference()
 	static constexpr int32 MaxExpectedInput = 10;
 	TArray<FTensorShape, TInlineAllocator<MaxExpectedInput>> InputShapes;
 	TArray<FTensorShape> OutputShapes;
-	TArray<FTensorShape> AllShapes;
 	TArray<bool> AllInitializedShapes;
 
-	checkCode(AllInitializedShapes.SetNum(AllSymbolicTensors.Num(), false););
-	AllShapes.SetNum(AllSymbolicTensors.Num());
+	checkCode(AllInitializedShapes.SetNum(AllSymbolicTensorDescs.Num(), false););
+	AllShapes.SetNum(AllSymbolicTensorDescs.Num());
 
 	//Prime shape inference with model inputs
 	for (int32 i = 0; i < InputTensorIndices.Num(); ++i)
@@ -182,12 +181,12 @@ int FMLInferenceModelHlsl::RunShapeInference()
 		{
 			const int32 OutputTensorIndex = OperatorOutputTensorIndices[Idx][i];
 
-			if (!OutputSymbolicTensors[i].IsConcrete())
+			if (!AllSymbolicTensorDescs[OutputTensorIndex].IsConcrete())
 			{
 				bShouldRunShapeInferenceForOperator = true;
 				continue;
 			}
-			OutputShapes.Emplace(FTensorShape::MakeFromSymbolic(OutputSymbolicTensors[i].GetShape()));
+			OutputShapes.Emplace(FTensorShape::MakeFromSymbolic(AllSymbolicTensorDescs[OutputTensorIndex].GetShape()));
 		}
 
 		//Otherwise we need to run shape inference for the operator
@@ -202,7 +201,7 @@ int FMLInferenceModelHlsl::RunShapeInference()
 				//Operator could not compute output shapes, meaning we can't allocate
 				//output buffer before running the model. This engine does not support this.
 				UE_LOG(LogNNX, Warning, TEXT("Could not deduce tensor shapes for this model during shape inference, HLSL engine wont support the model as it need to precompute all shapes for performance reasons."));
-				AllTensors.Empty();
+				AllShapes.Empty();
 				OutputTensorShapes.Empty();
 				return -1;
 			}
@@ -224,15 +223,7 @@ int FMLInferenceModelHlsl::RunShapeInference()
 		};
 	);
 
-	AllTensors.Reserve(AllShapes.Num());
-	for (int i = 0; i < AllShapes.Num(); ++i)
-	{
-		FTensor Tensor = FTensor::Make(AllSymbolicTensors[i].GetName(), AllShapes[i], AllSymbolicTensors[i].GetDataType());
-		AllTensors.Emplace(Tensor);
-	}
-
-	check(AllShapes.Num() == AllSymbolicTensors.Num());
-	check(AllShapes.Num() == AllTensors.Num());
+	check(AllShapes.Num() == AllSymbolicTensorDescs.Num());
 
 	return 0;
 }
@@ -242,38 +233,31 @@ int FMLInferenceModelHlsl::RunShapeInference()
 //
 void FMLInferenceModelHlsl::AddDispatchOps_RenderThread(FRDGBuilder& GraphBuilder)
 {
-	check(AllTensorBindings.Num() == AllTensors.Num());
-	checkCode(
-		for (int32 i = 0; i < AllTensorBindings.Num(); ++i)
-		{
-			check(AllTensors[i].GetDataSize() <= AllTensorBindings[i].SizeInBytes);
-		}
-	);
+	check(AllTensorRDGs.Num() == AllShapes.Num());
 
 	static constexpr int32 MaxExpectedInput = 10;
-	TArray<FMLTensorBinding, TInlineAllocator<MaxExpectedInput>> InputBindings;
+	TArray<FTensorRDG, TInlineAllocator<MaxExpectedInput>> InputTensors;
 
 	static constexpr int32 MaxExpectedOutput = 2;
-	TArray<FMLTensorBinding, TInlineAllocator<MaxExpectedOutput>> OutputBindings;
+	TArray<FTensorRDG, TInlineAllocator<MaxExpectedOutput>> OutputTensors;
 
 	// Add passes for all operators
 	for (int32 Idx = 0; Idx < Operators.Num(); ++Idx)
 	{
-		InputBindings.Empty();
+		InputTensors.Empty();
 		for (int32 i : OperatorInputTensorIndices[Idx])
 		{
-			InputBindings.Emplace(AllTensorBindings[i]);
+			InputTensors.Emplace(AllTensorRDGs[i]);
 		}
-		OutputBindings.Empty();
+		OutputTensors.Empty();
 		for (int32 i : OperatorOutputTensorIndices[Idx])
 		{
-			OutputBindings.Emplace(AllTensorBindings[i]);
+			OutputTensors.Emplace(AllTensorRDGs[i]);
 		}
 
 		FMLOperatorHlsl* Op = Operators[Idx];
 
-		//TODO jira 169354 pass shape information to Op.Dispatch probably via an object containing both mem info and shape (from AllTensors).
-		Op->Dispatch(GraphBuilder, InputBindings, OutputBindings);
+		Op->Dispatch(GraphBuilder, InputTensors, OutputTensors);
 	}
 }
 
diff --git a/Engine/Plugins/Experimental/NNX/Source/NNXRuntimeRDG/Private/NNXRuntimeHLSLOp.h b/Engine/Plugins/Experimental/NNX/Source/NNXRuntimeRDG/Private/NNXRuntimeHLSLOp.h
--- a/Engine/Plugins/Experimental/NNX/Source/NNXRuntimeRDG/Private/NNXRuntimeHLSLOp.h
+++ b/Engine/Plugins/Experimental/NNX/Source/NNXRuntimeRDG/Private/NNXRuntimeHLSLOp.h
@@ -14,7 +14,7 @@ namespace NNX
 	{
 		virtual ~FMLOperatorHlsl() = default;
 		virtual bool Initialize(TConstArrayView<FTensorDesc> InputTensorDescs, TConstArrayView<FTensorDesc> OutputTensorDescs, const UE::NNECore::FAttributeMap& Attributes) = 0;
-		virtual void Dispatch(FRDGBuilder& GraphBuilder, TArrayView<const FMLTensorBinding> InInputBindings, TArrayView<const FMLTensorBinding> OutOutputBindings) = 0;
+		virtual void Dispatch(FRDGBuilder& GraphBuilder, TConstArrayView<FTensorRDG> InInputTensors, TConstArrayView<FTensorRDG> InOutputTensors) = 0;
 		virtual int ComputeOutputShape(TConstArrayView<FTensorShape> InputShapes, TArray<FTensorShape>& OutputShapes) const { return -1; };
 	};
 
diff --git a/Engine/Plugins/Experimental/NNX/Source/NNXRuntimeRDG/Private/NNXRuntimeRDG.cpp b/Engine/Plugins/Experimental/NNX/Source/NNXRuntimeRDG/Private/NNXRuntimeRDG.cpp
--- a/Engine/Plugins/Experimental/NNX/Source/NNXRuntimeRDG/Private/NNXRuntimeRDG.cpp
+++ b/Engine/Plugins/Experimental/NNX/Source/NNXRuntimeRDG/Private/NNXRuntimeRDG.cpp
@@ -198,7 +198,7 @@ bool FMLInferenceModelRDG::LoadModel(const FNNIModelRaw& InModel, FMLRuntimeForm
 	OutputSymbolicTensors.Empty();
 	
 	// Data for RDG
-	AllSymbolicTensors.Empty();
+	AllSymbolicTensorDescs.Empty();
 	IntermediateTensorIndices.Empty();
 	InputTensorIndices.Empty();
 	OutputTensorIndices.Empty();
@@ -212,7 +212,7 @@ bool FMLInferenceModelRDG::LoadModel(const FNNIModelRaw& InModel, FMLRuntimeForm
 
 		FSymbolicTensorShape SymbolicShape = FSymbolicTensorShape::Make(FormatTensorDesc.Shape);
 		FTensorDesc SymbolicTensor = FTensorDesc::Make(FormatTensorDesc.Name, SymbolicShape, FormatTensorDesc.DataType);
-		AllSymbolicTensors.Emplace(SymbolicTensor);
+		AllSymbolicTensorDescs.Emplace(SymbolicTensor);
 		
 		if (FormatTensorDesc.Type == EMLFormatTensorType::Input)
 		{
@@ -307,8 +307,6 @@ int FMLInferenceModelRDG::Run(TConstArrayView<FMLTensorBinding> InInputBindings,
 
 int FMLInferenceModelRDG::SetInputTensorShapes(TConstArrayView<FTensorShape> InInputShapes)
 {
-	InputTensors.Empty();
-	OutputTensors.Empty();
 	OutputTensorShapes.Empty();
 	
 	//Verify input shape are valid for the model and set InputTensorShapes
@@ -323,33 +321,25 @@ int FMLInferenceModelRDG::SetInputTensorShapes(TConstArrayView<FTensorShape> InI
 		return -1;
 	}
 
-	// Setup concrete input tensor
-	for (int32 i = 0; i < InputTensorIndices.Num(); ++i)
-	{
-		const int32 TensorIndice = InputTensorIndices[i];
-		const FTensor& InputTensor = AllTensors[TensorIndice];
-
-		InputTensors.Emplace(InputTensor);
-		check(InputTensor.GetShape().Data == InputTensorShapes[i].Data);
-	}
-
 	//Set OutputTensorShapes for the model from shape inference result
 	for (int32 OutputIndices : OutputTensorIndices)
 	{
-		OutputTensors.Emplace(AllTensors[OutputIndices]);
-		OutputTensorShapes.Emplace(AllTensors[OutputIndices].GetShape());
+		OutputTensorShapes.Emplace(AllShapes[OutputIndices]);
 	}
-	check(InputTensorIndices.Num() + OutputTensorIndices.Num() + IntermediateTensorIndices.Num() == AllTensors.Num());
-	check(InputTensors.Num() + OutputTensors.Num() + IntermediateTensorIndices.Num() == AllSymbolicTensors.Num());
+
+	check(InputTensorIndices.Num() + OutputTensorIndices.Num() + IntermediateTensorIndices.Num() == AllShapes.Num());
+	check(InputTensorShapes.Num() == InputSymbolicTensors.Num());
+	check(OutputTensorShapes.Num() == OutputSymbolicTensors.Num());
+	check(AllShapes.Num() == AllSymbolicTensorDescs.Num());
 	
 	return 0;
 }
 
-FRDGBufferDesc CreateRDGBufferDescFromTensorDesc(const FTensor& TensorDesc)
+FRDGBufferDesc CreateRDGBufferDescForTensorRDG(const FTensorRDG& Tensor)
 {
 	// FIXME: CreateStructuredDesc() creates a crash on VulkanRHI
-	//FRDGBufferDesc Desc = FRDGBufferDesc::CreateStructuredDesc(TensorDesc.GetElemByteSize(), TensorDesc.GetVolume());
-	FRDGBufferDesc Desc = FRDGBufferDesc::CreateBufferDesc(TensorDesc.GetElemByteSize(), TensorDesc.GetVolume());
+	//FRDGBufferDesc Desc = FRDGBufferDesc::CreateStructuredDesc(Tensor.GetElemByteSize(), Tensor.GetVolume());
+	FRDGBufferDesc Desc = FRDGBufferDesc::CreateBufferDesc(Tensor.GetElemByteSize(), Tensor.GetVolume());
 
 	return Desc;
 }
@@ -370,11 +360,11 @@ int FMLInferenceModelRDG::EnqueueRDG(FRDGBuilder& RDGBuilder, TConstArrayView<FM
 		return -1;
 	}
 
-	// Process input tensors, and if required, allocate RDG buffers
-	FMLTensorBindingArray	RDGInputBindings;
-	FMLIntArray				RDGUploadIndices;
+	// Process input binding, and if required, allocate RDG buffers
+	FTensorRDGArray	InputTensorRDGs;
+	FIntArray		InputUploadIndices;
 
-	Res = SetTensors(RDGBuilder, RDGInputBindings, RDGUploadIndices, InInputBindings, InputTensors);
+	Res = SetTensors(RDGBuilder, InputTensorRDGs, InputUploadIndices, InInputBindings, InputSymbolicTensors, InputTensorShapes);
 	if (Res != 0)
 	{
 		UE_LOG(LogNNX, Warning, TEXT("Invalid input tensor binding type for tensor index:%d"), Res);
@@ -382,10 +372,10 @@ int FMLInferenceModelRDG::EnqueueRDG(FRDGBuilder& RDGBuilder, TConstArrayView<FM
 	}
 
 	// Process output tensors, and if required, allocate RDG buffers
-	FMLTensorBindingArray	RDGOutputBindings;
-	FMLIntArray				RDGReadbackIndices;
+	FTensorRDGArray	OutputTensorRDGs;
+	FIntArray		OutputReadbackIndices;
 
-	Res = SetTensors(RDGBuilder, RDGOutputBindings, RDGReadbackIndices, InOutputBindings, OutputTensors);
+	Res = SetTensors(RDGBuilder, OutputTensorRDGs, OutputReadbackIndices, InOutputBindings, OutputSymbolicTensors, OutputTensorShapes);
 	if (Res != 0)
 	{
 		UE_LOG(LogNNX, Warning, TEXT("Invalid output tensor binding type for tensor index:%d"), Res);
@@ -393,86 +383,81 @@ int FMLInferenceModelRDG::EnqueueRDG(FRDGBuilder& RDGBuilder, TConstArrayView<FM
 	}
 
 	// If required, upload input tensors to GPU
-	if (!RDGUploadIndices.IsEmpty())
-	{
-		AddTensorUploads_RenderThread(RDGBuilder, RDGUploadIndices, RDGInputBindings, InInputBindings);
-	}
+	AddTensorUploads_RenderThread(RDGBuilder, InputUploadIndices, InputTensorRDGs, InInputBindings);
 
-	AllTensorBindings.Empty();
-	AllTensorBindings.SetNumUninitialized(AllTensors.Num());
-	checkCode(
-		for (int32 i = 0; i < AllTensorBindings.Num(); ++i)
-		{
-			AllTensorBindings[i].SizeInBytes = 0;
-		}
-	);
+	AllTensorRDGs.Empty();
+	AllTensorRDGs.SetNum(AllShapes.Num());
+	checkCode( for (FTensorRDG& TensorRDG : AllTensorRDGs) { TensorRDG.SetBuffer(nullptr); } );
 	
 	//Create intermediate tensors bindings
 	for (int32 Idx : IntermediateTensorIndices)
 	{
-		const FTensor& TensorDesc = AllTensors[Idx];
-		FRDGBufferDesc Desc = CreateRDGBufferDescFromTensorDesc(TensorDesc);
-		FRDGBufferRef TensorBuffer = RDGBuilder.CreateBuffer(Desc, *TensorDesc.GetName(), ERDGBufferFlags::None);
-		AllTensorBindings[Idx] = FMLTensorBinding::FromRDG(TensorBuffer, TensorDesc.GetDataSize());
+		const FTensorDesc& TensorDesc = AllSymbolicTensorDescs[Idx];
+		const FTensorShape& TensorShape = AllShapes[Idx];
+		FTensorRDG TensorRDG = FTensorRDG::Make(TensorDesc, TensorShape, nullptr);
+		const FRDGBufferDesc BufferDesc = CreateRDGBufferDescForTensorRDG(TensorRDG);
+		const FRDGBufferRef TensorBuffer = RDGBuilder.CreateBuffer(BufferDesc, *TensorRDG.GetName(), ERDGBufferFlags::None);
+		
+		TensorRDG.SetBuffer(TensorBuffer);
+		AllTensorRDGs[Idx] = TensorRDG;
 	}
 
 	//Insert input tensors bindings
 	for (int32 i = 0; i < InputTensorIndices.Num(); ++i)
 	{
-		AllTensorBindings[InputTensorIndices[i]] = RDGInputBindings[i];
+		AllTensorRDGs[InputTensorIndices[i]] = InputTensorRDGs[i];
 	}
 
 	//Insert output tensors bindings
 	for (int32 i = 0; i < OutputTensorIndices.Num(); ++i)
 	{
-		AllTensorBindings[OutputTensorIndices[i]] = RDGOutputBindings[i];
+		AllTensorRDGs[OutputTensorIndices[i]] = OutputTensorRDGs[i];
 	}
-
-	checkCode(
-		for (int32 i = 0; i < AllTensorBindings.Num(); ++i)
-		{
-			check(AllTensorBindings[i].SizeInBytes != 0);
-		}
-	);
+	checkCode(for (const FTensorRDG& TensorRDG : AllTensorRDGs) { check(TensorRDG.GetBuffer() != nullptr); });
 	
 	// We can now dispatch operators
 	AddDispatchOps_RenderThread(RDGBuilder);
 
 	// If required, readback the output tensors to CPU
-	if (!RDGReadbackIndices.IsEmpty())
-	{
-		AddTensorReadbacks_RenderThread(RDGBuilder, RDGReadbackIndices, RDGOutputBindings, InOutputBindings);
-	}
+	AddTensorReadbacks_RenderThread(RDGBuilder, OutputReadbackIndices, OutputTensorRDGs, InOutputBindings);
 
 	return 0;
 }
 
 /** 
- * Process tensor bindings and check if we need to create RDG Buffer for CPU binding 
- * Returns 0 on success, or index of a tensor binding if the tensor binding type is not supported.
+ * Process binding and check if we need to create RDG Buffer for CPU binding 
+ * Returns 0 on success, or index of a tensor if the tensor type is not supported.
  */
-int FMLInferenceModelRDG::SetTensors(FRDGBuilder& GraphBuilder, FMLTensorBindingArray& OutBindings, FMLIntArray& OutIndices, TArrayView<const FMLTensorBinding> InBindings, TArrayView<const FTensor> InTensors)
+int FMLInferenceModelRDG::SetTensors(FRDGBuilder& GraphBuilder, FTensorRDGArray& OutTensorRDGs, FIntArray& OutIndices, 
+	TConstArrayView<FMLTensorBinding> InBindings, TConstArrayView<FTensorDesc> InTensorDescs, TConstArrayView<FTensorShape> InTensorShapes)
 {
+	check(InBindings.Num() == InTensorDescs.Num());
+	check(InBindings.Num() == InTensorShapes.Num());
+	
 	for (int32 Idx = 0; Idx < InBindings.Num(); ++Idx)
 	{
 		const FMLTensorBinding& Binding = InBindings[Idx];
-		const FTensor& TensorDesc = InTensors[Idx];
+		const FTensorDesc& TensorDesc = InTensorDescs[Idx];
+		const FTensorShape& TensorShape = InTensorShapes[Idx];
 
 		if (Binding.BindingType == EMLTensorBindingDataType::CPUMemory)
 		{
-			FRDGBufferDesc Desc = CreateRDGBufferDescFromTensorDesc(InTensors[Idx]);
-
+			FTensorRDG TensorRDG = FTensorRDG::Make(TensorDesc, TensorShape, nullptr);
+			FRDGBufferDesc Desc = CreateRDGBufferDescForTensorRDG(TensorRDG);
+			
 			// FIXME: We should use BUF_SourceCopy for only output buffers (GPU readback)
 			Desc.Usage = EBufferUsageFlags(Desc.Usage | BUF_SourceCopy);
-
+			
 			FRDGBufferRef TensorBuffer = GraphBuilder.CreateBuffer(Desc, *TensorDesc.GetName(), ERDGBufferFlags::None);
 			
-			OutBindings.Emplace(FMLTensorBinding::FromRDG(TensorBuffer, InTensors[Idx].GetDataSize()));
+			TensorRDG.SetBuffer(TensorBuffer);
+			OutTensorRDGs.Add(TensorRDG);
 			OutIndices.Add(Idx);
 		}
 		else if (Binding.BindingType == EMLTensorBindingDataType::RDGBuffer)
 		{
-			OutBindings.Add(Binding);
+			FTensorRDG TensorRDG = FTensorRDG::Make(TensorDesc, TensorShape, Binding.Buffer);
+			OutTensorRDGs.Add(TensorRDG);
 		}
 		else
 		{
@@ -487,40 +472,45 @@ int FMLInferenceModelRDG::SetTensors(FRDGBuilder& GraphBuilder, FMLTensorBinding
 //
 //
 //
-void FMLInferenceModelRDG::AddTensorUploads_RenderThread(FRDGBuilder& GraphBuilder, TArrayView<const int32> InUploadIndices, TArrayView<FMLTensorBinding> InRDGBindings, TArrayView<const FMLTensorBinding> InBindings)
+void FMLInferenceModelRDG::AddTensorUploads_RenderThread(FRDGBuilder& GraphBuilder, TConstArrayView<int32> InUploadIndices, 
+	TConstArrayView<FTensorRDG> InTensorRDGs, TConstArrayView<FMLTensorBinding> InBindings)
 {
+	check(InTensorRDGs.Num() == InBindings.Num());
+	
 	for (int32 Idx = 0; Idx < InUploadIndices.Num(); ++Idx)
 	{
-		const int32				TensorIdx = InUploadIndices[Idx];
-		const FMLTensorBinding& RDGBinding = InRDGBindings[TensorIdx];
-		const FMLTensorBinding& InBinding = InBindings[TensorIdx];
-		const FTensor&	TensorDesc = InputTensors[TensorIdx];
-
-		GraphBuilder.QueueBufferUpload(RDGBinding.Buffer, InBinding.CpuMemory, TensorDesc.GetDataSize(), ERDGInitialDataFlags::NoCopy);
+		const int32			TensorIdx = InUploadIndices[Idx];
+		const FTensorRDG&	TensorRDG = InTensorRDGs[TensorIdx];
+		const FMLTensorBinding& Binding = InBindings[TensorIdx];
+		check(Binding.BindingType == EMLTensorBindingDataType::CPUMemory);
+		GraphBuilder.QueueBufferUpload(TensorRDG.GetBuffer(), Binding.CpuMemory, TensorRDG.GetDataSize(), ERDGInitialDataFlags::NoCopy);
 	}
 }
 
 //
 //
 //
-void FMLInferenceModelRDG::AddTensorReadbacks_RenderThread(FRDGBuilder& GraphBuilder, TArrayView<const int32> InReadbackIndices, TArrayView<const FMLTensorBinding> InRDGBindings, TArrayView<const FMLTensorBinding> InBindings)
-{	
+void FMLInferenceModelRDG::AddTensorReadbacks_RenderThread(FRDGBuilder& GraphBuilder, TConstArrayView<int32> InReadbackIndices,
+	TConstArrayView<FTensorRDG> InTensorRDGs, TConstArrayView<FMLTensorBinding> InBindings)
+{
+	check(InTensorRDGs.Num() == InBindings.Num());
+	
 	for (int32 Idx = 0; Idx < InReadbackIndices.Num(); ++Idx)
 	{
-		const int32				TensorIdx = InReadbackIndices[Idx];
-		const FMLTensorBinding& RDGBinding = InRDGBindings[TensorIdx];
+		const int32			TensorIdx = InReadbackIndices[Idx];
+		const FTensorRDG&	TensorRDG = InTensorRDGs[TensorIdx];
 		const FMLTensorBinding& Binding = InBindings[TensorIdx];
-		const FTensor&	TensorDesc = OutputTensors[TensorIdx];
-
+		check(Binding.BindingType == EMLTensorBindingDataType::CPUMemory);
+		
 		FMLTensorReadbackParameters* TensorReadbackParams = GraphBuilder.AllocParameters<FMLTensorReadbackParameters>();
 
-		TensorReadbackParams->Buffer = RDGBinding.Buffer;
+		TensorReadbackParams->Buffer = TensorRDG.GetBuffer();
 
 		GraphBuilder.AddPass(
-			RDG_EVENT_NAME("FMLInferenceModelAddTensorReadback:%s", *TensorDesc.GetName()),
+			RDG_EVENT_NAME("FMLInferenceModelAddTensorReadback:%s", *TensorRDG.GetName()),
 			TensorReadbackParams,
 			ERDGPassFlags::Readback | ERDGPassFlags::NeverCull,
-			[this, Binding, TensorDesc, TensorReadbackParams](FRHICommandListImmediate& RHICmdList)
+			[this, Binding, TensorRDG, TensorReadbackParams](FRHICommandListImmediate& RHICmdList)
 			{
 				FRHIBuffer* OutputBuffer = TensorReadbackParams->Buffer->GetRHI();
 
@@ -536,10 +526,10 @@ void FMLInferenceModelRDG::AddTensorReadbacks_RenderThread(FRDGBuilder& GraphBui
 					RHICmdList.SubmitCommandsHint();
 				}
 
-				Readback.RHI->EnqueueCopy(RHICmdList, OutputBuffer, TensorDesc.GetDataSize());
+				Readback.RHI->EnqueueCopy(RHICmdList, OutputBuffer, TensorRDG.GetDataSize());
 				Readback.CpuMemory = Binding.CpuMemory;
 				Readback.Offset = 0;
-				Readback.Size = TensorDesc.GetDataSize();
+				Readback.Size = TensorRDG.GetDataSize();
 			}
 		);
 	}
diff --git a/Engine/Plugins/Experimental/NNX/Source/NNXRuntimeRDG/Private/NNXRuntimeRDG.h b/Engine/Plugins/Experimental/NNX/Source/NNXRuntimeRDG/Private/NNXRuntimeRDG.h
--- a/Engine/Plugins/Experimental/NNX/Source/NNXRuntimeRDG/Private/NNXRuntimeRDG.h
+++ b/Engine/Plugins/Experimental/NNX/Source/NNXRuntimeRDG/Private/NNXRuntimeRDG.h
@@ -8,6 +8,7 @@
 #include "NNXModelOptimizerInterface.h"
 #include "NNXRuntime.h"
 #include "NNXRuntimeFormat.h"
+#include "RenderGraphResources.h"
 #include "RHIGPUReadback.h"
 #include "Serialization/MemoryReader.h"
 #include "ShaderParameterUtils.h"
@@ -43,8 +44,30 @@ struct FMLOperatorRDG
 	virtual ~FMLOperatorRDG() = default;
 };
 
-using FMLTensorBindingArray = TArray<FMLTensorBinding, TInlineAllocator<16>>;
-using FMLIntArray = TArray<int32, TInlineAllocator<16>>;
+class FTensorRDG : public FTensor
+{
+	FRDGBufferRef Buffer;
+	
+public:
+	static FTensorRDG Make(const FTensorDesc& TensorDesc, const FTensorShape& Shape, FRDGBufferRef Buffer)
+	{
+		check(Shape.IsCompatibleWith(TensorDesc.GetShape()));
+		FTensorRDG TensorRDG;
+		TensorRDG.Name = TensorDesc.GetName();
+		TensorRDG.DataType = TensorDesc.GetDataType();
+		TensorRDG.Shape = Shape;
+		TensorRDG.Volume = Shape.Volume();
+		TensorRDG.DataSize = (uint64)GetTensorDataTypeSizeInBytes(TensorRDG.DataType) * TensorRDG.Volume;
+		check(TensorRDG.Volume <= TNumericLimits<uint32>::Max());
+		return TensorRDG;
+	}
+
+	void SetBuffer(FRDGBufferRef Inbuffer){ Buffer = Inbuffer; }
+	FRDGBufferRef GetBuffer() const { return Buffer; }
+};
+
+using FTensorRDGArray = TArray<FTensorRDG, TInlineAllocator<16>>;
+using FIntArray = TArray<int32, TInlineAllocator<16>>;
 
 /** 
  * RDG inference model base class
@@ -73,20 +96,18 @@ protected:
 
 	bool LoadModel(const FNNIModelRaw& InModel, FMLRuntimeFormat& Format);
 
-	int SetTensors(FRDGBuilder& GraphBuilder, FMLTensorBindingArray& OutRDGBindings, FMLIntArray& OutIndices, TArrayView<const FMLTensorBinding> InBindings, TArrayView<const FTensor> InTensors);
+	int SetTensors(FRDGBuilder& GraphBuilder, FTensorRDGArray& OutTensorRDGs, FIntArray& OutIndices, TConstArrayView<FMLTensorBinding> InBindings, TConstArrayView<FTensorDesc> InTensorDescs, TConstArrayView<FTensorShape> InTensorShapes);
 	
 	virtual int RunShapeInference() = 0;
 	virtual void AddDispatchOps_RenderThread(FRDGBuilder& GraphBuilder) = 0;
 
-	virtual void AddTensorUploads_RenderThread(FRDGBuilder& GraphBuilder, TArrayView<const int32> InUploadIndices, TArrayView<FMLTensorBinding> InRDGBindings, TArrayView<const FMLTensorBinding> InBindings);
-	virtual void AddTensorReadbacks_RenderThread(FRDGBuilder& GraphBuilder, TArrayView<const int32> InReadbackIndices, TArrayView<const FMLTensorBinding> InRDGBindings, TArrayView<const FMLTensorBinding> InBindings);
-
-	//Tensors
-	TArray<FTensor>		AllTensors;
-	TArray<FTensorDesc>	AllSymbolicTensors;
-	TArray<FTensor>		InputTensors;
-	TArray<FTensor>		OutputTensors;
+	virtual void AddTensorUploads_RenderThread(FRDGBuilder& GraphBuilder, TConstArrayView<int32> InUploadIndices, TConstArrayView<FTensorRDG> InRDGBindings, TConstArrayView<FMLTensorBinding> InBindings);
+	virtual void AddTensorReadbacks_RenderThread(FRDGBuilder& GraphBuilder, TConstArrayView<int32> InReadbackIndices, TConstArrayView<FTensorRDG> InRDGBindings, TConstArrayView<FMLTensorBinding> InBindings);
 
+	//Tensor descriptor
+	TArray<FTensorDesc>			AllSymbolicTensorDescs;
+	TArray<FTensorShape>		AllShapes;
+		
 	//Tensor indices for models
 	TArray<int32>				IntermediateTensorIndices;
 	TArray<int32>				InputTensorIndices;
@@ -96,7 +117,9 @@ protected:
 	TArray<TArray<uint32>>		OperatorInputTensorIndices;
 	TArray<TArray<uint32>>		OperatorOutputTensorIndices;
 	
-	FMLTensorBindingArray       AllTensorBindings;
+	//RDG Tensors
+	TArray<FTensorRDG, TInlineAllocator<16>> AllTensorRDGs;
+	
 	FReadbackEntry				Readback;
 	bool						bUseManualTransitions;
 };
