Extend support for more attribute types and add attributes to Convolution operators.

diff --git a/Engine/Plugins/Experimental/NNX/Source/NNEHlslShaders/Internal/NNEHlslShadersConvCS.h b/Engine/Plugins/Experimental/NNX/Source/NNEHlslShaders/Internal/NNEHlslShadersConvCS.h
--- a/Engine/Plugins/Experimental/NNX/Source/NNEHlslShaders/Internal/NNEHlslShadersConvCS.h
+++ b/Engine/Plugins/Experimental/NNX/Source/NNEHlslShaders/Internal/NNEHlslShadersConvCS.h
@@ -105,5 +105,7 @@ namespace UE::NNEHlslShaders::Internal
 		static FIntVector GetGroupCount(TArrayView<const int32> YShape, TArrayView<const int32> GroupShape);
 
 		static EConvGroupSize GetMinimalGroupSize(TArrayView<const int32> WShape);
+
+		static void LexFromString(EConvAutoPad& OutValue, const TCHAR* StringVal);
 	};
 } // UE::NNEHlslShaders::Internal
\ No newline at end of file
diff --git a/Engine/Plugins/Experimental/NNX/Source/NNEHlslShaders/Internal/NNEHlslShadersConvTransposeCS.h b/Engine/Plugins/Experimental/NNX/Source/NNEHlslShaders/Internal/NNEHlslShadersConvTransposeCS.h
--- a/Engine/Plugins/Experimental/NNX/Source/NNEHlslShaders/Internal/NNEHlslShadersConvTransposeCS.h
+++ b/Engine/Plugins/Experimental/NNX/Source/NNEHlslShaders/Internal/NNEHlslShadersConvTransposeCS.h
@@ -95,6 +95,8 @@ namespace UE::NNEHlslShaders::Internal
 
 		static EConvTransposeGroupSize GetMinimalGroupSize(TArrayView<const uint32> WShape);
 
+		static void LexFromString(EConvTransposeAutoPad& OutValue, const TCHAR* StringVal);
+		
 	private:
 
 		static TArray<int32> GetXBlockShape(TArrayView<const int32> GroupShape, TArrayView<const uint32> WShape, TArrayView<const int32> Dilations, TArrayView<const int32> Strides);
diff --git a/Engine/Plugins/Experimental/NNX/Source/NNEHlslShaders/Private/NNEHlslShadersConvCS.cpp b/Engine/Plugins/Experimental/NNX/Source/NNEHlslShaders/Private/NNEHlslShadersConvCS.cpp
--- a/Engine/Plugins/Experimental/NNX/Source/NNEHlslShaders/Private/NNEHlslShadersConvCS.cpp
+++ b/Engine/Plugins/Experimental/NNX/Source/NNEHlslShaders/Private/NNEHlslShadersConvCS.cpp
@@ -288,5 +288,14 @@ namespace UE::NNEHlslShaders::Internal
 		return EConvGroupSize::MAX;
 	}
 
+	void FConvCS::LexFromString(EConvAutoPad& OutValue, const TCHAR* StringVal)
+	{
+		OutValue = EConvAutoPad::NOTSET;
+		if (FCString::Stricmp(StringVal, TEXT("NOTSET")) == 0) OutValue = EConvAutoPad::NOTSET;
+		else if (FCString::Stricmp(StringVal, TEXT("SAME_UPPER")) == 0) OutValue = EConvAutoPad::SAME_UPPER;
+		else if (FCString::Stricmp(StringVal, TEXT("SAME_LOWER")) == 0) OutValue = EConvAutoPad::SAME_LOWER;
+		else if (FCString::Stricmp(StringVal, TEXT("VALID")) == 0) OutValue = EConvAutoPad::VALID;
+	}
+
 	IMPLEMENT_GLOBAL_SHADER(FConvCS, "/NNE/NNEHlslShadersConv.usf", "Conv", SF_Compute);
 } // UE::NNEHlslShaders::Internal
\ No newline at end of file
diff --git a/Engine/Plugins/Experimental/NNX/Source/NNEHlslShaders/Private/NNEHlslShadersConvTransposeCS.cpp b/Engine/Plugins/Experimental/NNX/Source/NNEHlslShaders/Private/NNEHlslShadersConvTransposeCS.cpp
--- a/Engine/Plugins/Experimental/NNX/Source/NNEHlslShaders/Private/NNEHlslShadersConvTransposeCS.cpp
+++ b/Engine/Plugins/Experimental/NNX/Source/NNEHlslShaders/Private/NNEHlslShadersConvTransposeCS.cpp
@@ -281,5 +281,14 @@ namespace UE::NNEHlslShaders::Internal
 		return Result;
 	}
 
+	void FConvTransposeCS::LexFromString(EConvTransposeAutoPad& OutValue, const TCHAR* StringVal)
+	{
+		OutValue = EConvTransposeAutoPad::NOTSET;
+		if (FCString::Stricmp(StringVal, TEXT("NOTSET")) == 0) OutValue = EConvTransposeAutoPad::NOTSET;
+		else if (FCString::Stricmp(StringVal, TEXT("SAME_UPPER")) == 0) OutValue = EConvTransposeAutoPad::SAME_UPPER;
+		else if (FCString::Stricmp(StringVal, TEXT("SAME_LOWER")) == 0) OutValue = EConvTransposeAutoPad::SAME_LOWER;
+		else if (FCString::Stricmp(StringVal, TEXT("VALID")) == 0) OutValue = EConvTransposeAutoPad::VALID;
+	}
+
 	IMPLEMENT_GLOBAL_SHADER(FConvTransposeCS, "/NNE/NNEHlslShadersConvTranspose.usf", "ConvTranspose", SF_Compute);
 } // UE::NNEHlslShaders::Internal
\ No newline at end of file
diff --git a/Engine/Plugins/Experimental/NNX/Source/NNXCore/Private/NNECoreAttributeDataType.cpp b/Engine/Plugins/Experimental/NNX/Source/NNXCore/Private/NNECoreAttributeDataType.cpp
new file mode 100644
--- /dev/null
+++ b/Engine/Plugins/Experimental/NNX/Source/NNXCore/Private/NNECoreAttributeDataType.cpp
@@ -0,0 +1,15 @@
+// Copyright Epic Games, Inc. All Rights Reserved.
+
+#include "NNECoreAttributeDataType.h"
+
+void LexFromString(ENNEAttributeDataType& OutValue, const TCHAR* StringVal)
+{
+	int64 EnumVal = StaticEnum<ENNEAttributeDataType>()->GetValueByName(StringVal);
+	if (EnumVal == INDEX_NONE)
+	{
+		OutValue = ENNEAttributeDataType::None;
+		ensureMsgf(false, TEXT("ENNEAttributeDataType LexFromString didn't have a match for '%s'"), StringVal);
+	}
+
+	OutValue = (ENNEAttributeDataType)EnumVal;
+};
\ No newline at end of file
diff --git a/Engine/Plugins/Experimental/NNX/Source/NNXCore/Public/NNECoreAttributeDataType.h b/Engine/Plugins/Experimental/NNX/Source/NNXCore/Public/NNECoreAttributeDataType.h
--- a/Engine/Plugins/Experimental/NNX/Source/NNXCore/Public/NNECoreAttributeDataType.h
+++ b/Engine/Plugins/Experimental/NNX/Source/NNXCore/Public/NNECoreAttributeDataType.h
@@ -5,8 +5,6 @@
 #include "UObject/Class.h"
 #include "UObject/Object.h"
 
-#include "NNECoreAttributeDataType.generated.h"
-
 /**
  * Attribute data types
  * 
@@ -17,21 +15,13 @@ enum class ENNEAttributeDataType : uint8
 {
 	None,
 	Float,								//!< 32-bit floating number
-	Int32								//!< 32-bit signed integer
+	Int32,								//!< 32-bit signed integer
+	Int32Array,							//!< TArray of 32-bit signed integers
+	String								//!< built-in FString
 	//TODO add more AttributeDataType support
 };
 
 /**
  * @return ENNEAttributeDataType from the string passed in
  */
-inline void LexFromString(ENNEAttributeDataType& OutValue, const TCHAR* StringVal)
-{
-	int64 EnumVal = StaticEnum<ENNEAttributeDataType>()->GetValueByName(StringVal);
-	if (EnumVal == INDEX_NONE)
-	{
-		OutValue = ENNEAttributeDataType::None;
-		ensureMsgf(false, TEXT("ENNEAttributeDataType LexFromString didn't have a match for '%s'"), StringVal);
-	}
-
-	OutValue = (ENNEAttributeDataType)EnumVal;
-};
\ No newline at end of file
+NNXCORE_API void LexFromString(ENNEAttributeDataType& OutValue, const TCHAR* StringVal);
\ No newline at end of file
diff --git a/Engine/Plugins/Experimental/NNX/Source/NNXCore/Public/NNECoreAttributeValueTraits.h b/Engine/Plugins/Experimental/NNX/Source/NNXCore/Public/NNECoreAttributeValueTraits.h
--- a/Engine/Plugins/Experimental/NNX/Source/NNXCore/Public/NNECoreAttributeValueTraits.h
+++ b/Engine/Plugins/Experimental/NNX/Source/NNXCore/Public/NNECoreAttributeValueTraits.h
@@ -22,4 +22,14 @@ template<> struct TNNEAttributeValueTraits<float>
 template<> struct TNNEAttributeValueTraits<int32>
 {
 	static constexpr ENNEAttributeDataType GetType() { return ENNEAttributeDataType::Int32; }
+};
+
+template<> struct TNNEAttributeValueTraits<FString>
+{
+	static constexpr ENNEAttributeDataType GetType() { return ENNEAttributeDataType::String; }
+};
+
+template<> struct TNNEAttributeValueTraits<TArray<int32>>
+{
+	static constexpr ENNEAttributeDataType GetType() { return ENNEAttributeDataType::Int32Array; }
 };
\ No newline at end of file
diff --git a/Engine/Plugins/Experimental/NNX/Source/NNXQA/Private/NNXQAJsonUtils.h b/Engine/Plugins/Experimental/NNX/Source/NNXQA/Private/NNXQAJsonUtils.h
--- a/Engine/Plugins/Experimental/NNX/Source/NNXQA/Private/NNXQAJsonUtils.h
+++ b/Engine/Plugins/Experimental/NNX/Source/NNXQA/Private/NNXQAJsonUtils.h
@@ -97,6 +97,20 @@ namespace Json
 					Value = FNNEAttributeValue(TmpValue);
 					break;
 				}
+				case ENNEAttributeDataType::Int32Array:
+				{
+					TArray<int32> TmpValue;
+					JSON_SERIALIZE_ARRAY("value", TmpValue);
+					Value = FNNEAttributeValue(TmpValue);
+					break;
+				}
+				case ENNEAttributeDataType::String:
+				{
+					FString TmpValue;
+					JSON_SERIALIZE("value", TmpValue);
+					Value = FNNEAttributeValue(TmpValue);
+					break;
+				}
 				default:
 					check(Type == ENNEAttributeDataType::None);
 			}
diff --git a/Engine/Plugins/Experimental/NNX/Source/NNXQA/Private/NNXQAParametricTest.cpp b/Engine/Plugins/Experimental/NNX/Source/NNXQA/Private/NNXQAParametricTest.cpp
--- a/Engine/Plugins/Experimental/NNX/Source/NNXQA/Private/NNXQAParametricTest.cpp
+++ b/Engine/Plugins/Experimental/NNX/Source/NNXQA/Private/NNXQAParametricTest.cpp
@@ -213,6 +213,17 @@ namespace Test
 				{
 					case ENNEAttributeDataType::Float: return FString::SanitizeFloat(Value.GetValue<float>());
 					case ENNEAttributeDataType::Int32: return FString::FromInt(Value.GetValue<int32>());
+					case ENNEAttributeDataType::Int32Array:
+					{
+						TArray<FString> ArrayStrings;
+						for (int32 Val : Value.GetValue<TArray<int32>>())
+						{
+							ArrayStrings.Add(FString::FromInt(Val));
+						}
+
+						return TEXT("[") + FString::Join(ArrayStrings, TEXT(",")) + TEXT("]");
+					}
+					case ENNEAttributeDataType::String: return Value.GetValue<FString>();
 				}
 				return TEXT("-");
 			};
diff --git a/Engine/Plugins/Experimental/NNX/Source/NNXQA/Resources/NNXQATestDesc.json b/Engine/Plugins/Experimental/NNX/Source/NNXQA/Resources/NNXQATestDesc.json
--- a/Engine/Plugins/Experimental/NNX/Source/NNXQA/Resources/NNXQATestDesc.json
+++ b/Engine/Plugins/Experimental/NNX/Source/NNXQA/Resources/NNXQATestDesc.json
@@ -262,9 +262,23 @@
 		},
 		{
 			"category": "Convolution",
-			"runtimes": [	{	"name": "NNXRuntimeDml", "skip": true	}	],
+			"runtimes": [
+				{ "name": "NNXRuntimeDml", "skip": true	},
+				{ "name": "NNXRuntimeORTDml", "relative_error": 3.5e-2 },
+				{ "name": "NNXRuntimeHlsl",   "relative_error": 3.5e-2 }
+			],
+			"targets": [
+				{ "target": "Conv" }
+			]
+		},
+		{
+			"category": "ConvTranspose",
+			"runtimes": [
+				{ "name": "NNXRuntimeDml", "skip": true	},
+				{ "name": "NNXRuntimeORTDml", "relative_error": 3.5e-2 },
+				{ "name": "NNXRuntimeHlsl",   "relative_error": 3.5e-2 }
+			],
 			"targets": [
-				{ "target": "Conv" },
 				{ "target": "ConvTranspose" }
 			]
 		},
@@ -454,11 +468,33 @@
 				{ "inputs": [ { "shape": [ 1, 1 ] }, { "shape": [ 5, 1 ] }, { "shape": [ 5 ] } ], "outputs": [ { "shape": [ 1, 5 ] } ]}
 			]
 		},
-		{
-			"name": "Convolution",
+		{ 
+			"name": "Convolution.DefaultPadding",
 			"datasets": [
-				{	"inputs": [	{ "shape": [ 1, 1, 1 ] },	{ "shape": [ 1, 1, 1 ] } ],	"outputs": [ { "shape": [ 1, 1, 1 ] } ]},
-				{	"inputs": [	{ "shape": [ 1, 1, 1 ] },	{ "shape": [ 1, 1, 1 ] },	{ "shape": [ 1 ] } ],	"outputs": [ { "shape": [ 1, 1, 1 ] } ]}
+				{ "inputs": [ { "shape": [ 1, 1, 5, 5 ] }, { "shape": [ 1, 1, 3, 3 ] } ], "outputs": [ { "shape": [ 1, 1, 5, 5 ] } ] },
+				{ "inputs": [ { "shape": [ 1, 3, 5, 5 ] }, { "shape": [ 1, 3, 3, 3 ] } ], "outputs": [ { "shape": [ 1, 1, 5, 5 ] } ], "comment": "With channels" },
+				{ "inputs": [ { "shape": [ 1, 1, 5, 5 ] }, { "shape": [ 3, 1, 3, 3 ] } ], "outputs": [ { "shape": [ 1, 3, 5, 5 ] } ], "comment": "With out-channels" },
+				{ "inputs": [ { "shape": [ 1, 1, 5, 5 ] }, { "shape": [ 3, 1, 3, 3 ] }, { "shape": [ 3 ] } ], "outputs": [ { "shape": [ 1, 3, 5, 5 ] } ], "comment": "With bias" },
+				{ "inputs": [ { "shape": [ 3, 1, 5, 5 ] }, { "shape": [ 1, 1, 3, 3 ] } ], "outputs": [ { "shape": [ 3, 1, 5, 5 ] } ], "comment": "Batches" }
+			]
+		},
+		{ "name": "Convolution.WithoutOrAutoPadding", 			"datasets": [ { "inputs": [ { "shape": [ 1, 1, 5, 5 ] }, { "shape": [ 1, 1, 3, 3 ] } ], "outputs": [ { "shape": [ 1, 1, 3, 3 ] } ] } ] },
+		{ "name": "Convolution.WithStridesAndPadding", 			"datasets": [ { "inputs": [ { "shape": [ 1, 1, 7, 5 ] }, { "shape": [ 1, 1, 3, 3 ] } ], "outputs": [ { "shape": [ 1, 1, 4, 3 ] } ] } ] },
+		{ "name": "Convolution.WithoutPaddingAndStrides", 		"datasets": [ { "inputs": [ { "shape": [ 1, 1, 7, 5 ] }, { "shape": [ 1, 1, 3, 3 ] } ], "outputs": [ { "shape": [ 1, 1, 3, 2 ] } ] } ] },
+		{ "name": "Convolution.AsymmetricPaddingAndStrides", 	"datasets": [ { "inputs": [ { "shape": [ 1, 1, 7, 5 ] }, { "shape": [ 1, 1, 3, 3 ] } ], "outputs": [ { "shape": [ 1, 1, 4, 2 ] } ] } ] },
+		{ "name": "Convolution.WithGroups", 					"datasets": [ { "inputs": [ { "shape": [ 1, 3, 5, 5 ] }, { "shape": [ 3, 1, 3, 3 ] } ], "outputs": [ { "shape": [ 1, 3, 5, 5 ] } ] } ] },
+		{ "name": "Convolution.Large1", 						"datasets": [ { "inputs": [ { "shape": [ 1, 1, 17, 17 ] }, { "shape": [ 1, 1, 3, 3 ] } ], "outputs": [ { "shape": [ 1, 1, 17, 17 ] } ] } ] },
+		{ "name": "Convolution.Large2", 						"datasets": [ { "inputs": [ { "shape": [ 1, 1, 37, 37 ] }, { "shape": [ 1, 1, 3, 3 ] } ], "outputs": [ { "shape": [ 1, 1, 17, 17 ] } ] } ] },
+		{ "name": "ConvTranspose.Default", 		"datasets": [ { "inputs": [ { "shape": [ 1, 1, 3, 3 ] }, { "shape": [ 1, 2, 3, 3 ] } ], "outputs": [ { "shape": [ 1, 2, 5, 5 ] } ] } ] },
+		{ "name": "ConvTranspose.1D", 			"datasets": [ { "inputs": [ { "shape": [ 1, 1, 3 ] }, { "shape": [ 1, 2, 3 ] } ], "outputs": [ { "shape": [ 1, 2, 5 ] } ] } ] },
+		{ "name": "ConvTranspose.3D", 			"datasets": [ { "inputs": [ { "shape": [ 1, 1, 3, 4, 5 ] }, { "shape": [ 1, 2, 3, 3, 3 ] } ], "outputs": [ { "shape": [ 1, 2, 5, 6, 7 ] } ] } ] },
+		{ "name": "ConvTranspose.Attributes", 	"datasets": [ { "inputs": [ { "shape": [ 1, 1, 3, 3 ] }, { "shape": [ 1, 2, 3, 3 ] } ], "outputs": [ { "shape": [ 1, 2, 10, 8 ] } ] } ] },
+		{
+			"name": "ConvTranspose.Autopad",
+			"datasets": [ { "inputs": [ { "shape": [ 1, 1, 3, 3 ] }, { "shape": [ 1, 2, 3, 3 ] } ], "outputs": [ { "shape": [ 1, 2, 6, 6 ] } ] } ],
+			"runtimes": [
+				{ "name": "NNXRuntimeORTDml", "skip": true, "comment": "TODO INVESTIGATE Returns 'The parameter is incorrect'" },
+				{ "name": "NNXRuntimeHlsl", "skip": true, "comment": "Fails until ORT 1.13, because they implemented 'auto_pad' differently (they flipped 'SAME_LOWER' and 'SAME_UPPER')" }
 			]
 		},
 		{
@@ -550,6 +586,179 @@
 			"attribute_maps": [
 				{ "attributes": [ { "name": "transA", "value": 1, "type": "Int32" }, { "name": "transB", "value": 1, "type": "Int32" } ] }
 			]
+		},
+		{
+			"name": "Convolution.DefaultPadding",
+			"attribute_maps": [
+				{ "attributes": [
+					{ "name": "auto_pad", "value": "NOTSET", "type": "String" },
+					{ "name": "dilations", "value": [ 1, 1 ], "type": "Int32Array" },
+					{ "name": "group", "value": 1, "type": "Int32" },
+					{ "name": "pads", "value": [ 1, 1, 1, 1 ], "type": "Int32Array" },
+					{ "name": "strides", "value": [ 1, 1 ], "type": "Int32Array"}
+				] }
+			]
+		},
+		{
+			"name": "Convolution.WithoutOrAutoPadding.WithoutPadding",
+			"attribute_maps": [
+				{ "attributes": [
+					{ "name": "auto_pad", "value": "NOTSET", "type": "String" },
+					{ "name": "dilations", "value": [ 1, 1 ], "type": "Int32Array" },
+					{ "name": "group", "value": 1, "type": "Int32" },
+					{ "name": "pads", "value": [ 0, 0, 0, 0 ], "type": "Int32Array" },
+					{ "name": "strides", "value": [ 1, 1 ], "type": "Int32Array"}
+				] }
+			]
+		},
+		{
+			"name": "Convolution.WithoutOrAutoPadding.AutoPaddingAndStrides",
+			"attribute_maps": [
+				{ "attributes": [
+					{ "name": "auto_pad", "value": "SAME_LOWER", "type": "String" },
+					{ "name": "dilations", "value": [ 1, 1 ], "type": "Int32Array" },
+					{ "name": "group", "value": 1, "type": "Int32" },
+					{ "name": "strides", "value": [ 2, 2 ], "type": "Int32Array"}
+				] }
+			]
+		},
+		{
+			"name": "Convolution.WithStridesAndPadding",
+			"attribute_maps": [
+				{ "attributes": [
+					{ "name": "auto_pad", "value": "NOTSET", "type": "String" },
+					{ "name": "dilations", "value": [ 1, 1 ], "type": "Int32Array" },
+					{ "name": "group", "value": 1, "type": "Int32" },
+					{ "name": "pads", "value": [ 1, 1, 1, 1 ], "type": "Int32Array" },
+					{ "name": "strides", "value": [ 2, 2 ], "type": "Int32Array"}
+				] }
+			]
+		},
+		{
+			"name": "Convolution.WithoutPaddingAndStrides",
+			"attribute_maps": [
+				{ "attributes": [
+					{ "name": "auto_pad", "value": "NOTSET", "type": "String" },
+					{ "name": "dilations", "value": [ 1, 1 ], "type": "Int32Array" },
+					{ "name": "group", "value": 1, "type": "Int32" },
+					{ "name": "pads", "value": [ 0, 0, 0, 0 ], "type": "Int32Array" },
+					{ "name": "strides", "value": [ 2, 2 ], "type": "Int32Array"}
+				] }
+			]
+		},
+		{
+			"name": "Convolution.AsymmetricPaddingAndStrides",
+			"attribute_maps": [
+				{ "attributes": [
+					{ "name": "auto_pad", "value": "NOTSET", "type": "String" },
+					{ "name": "dilations", "value": [ 1, 1 ], "type": "Int32Array" },
+					{ "name": "group", "value": 1, "type": "Int32" },
+					{ "name": "pads", "value": [ 1, 0, 1, 0 ], "type": "Int32Array" },
+					{ "name": "strides", "value": [ 2, 2 ], "type": "Int32Array"}
+				] }
+			]
+		},
+		{
+			"name": "Convolution.WithGroups",
+			"attribute_maps": [
+				{ "attributes": [
+					{ "name": "auto_pad", "value": "NOTSET", "type": "String" },
+					{ "name": "dilations", "value": [ 1, 1 ], "type": "Int32Array" },
+					{ "name": "group", "value": 3, "type": "Int32" },
+					{ "name": "pads", "value": [ 1, 1, 1, 1 ], "type": "Int32Array" },
+					{ "name": "strides", "value": [ 1, 1 ], "type": "Int32Array"}
+				] }
+			]
+		},
+		{
+			"name": "Convolution.Large1",
+			"attribute_maps": [
+				{ "attributes": [
+					{ "name": "auto_pad", "value": "SAME_UPPER", "type": "String" },
+					{ "name": "dilations", "value": [ 1, 1 ], "type": "Int32Array" },
+					{ "name": "group", "value": 1, "type": "Int32" },
+					{ "name": "strides", "value": [ 1, 1 ], "type": "Int32Array"}
+				] }
+			]
+		},
+		{
+			"name": "Convolution.Large2",
+			"attribute_maps": [
+				{ "attributes": [
+					{ "name": "auto_pad", "value": "NOTSET", "type": "String" },
+					{ "name": "dilations", "value": [ 2, 2 ], "type": "Int32Array" },
+					{ "name": "group", "value": 1, "type": "Int32" },
+					{ "name": "pads", "value": [ 0, 0, 0, 0 ], "type": "Int32Array" },
+					{ "name": "strides", "value": [ 2, 2 ], "type": "Int32Array"}
+				] }
+			]
+		},
+		{
+			"name": "ConvTranspose.Default",
+			"attribute_maps": [
+				{ "attributes": [
+					{ "name": "auto_pad", "value": "NOTSET", "type": "String" },
+					{ "name": "dilations", "value": [ 1, 1 ], "type": "Int32Array" },
+					{ "name": "group", "value": 1, "type": "Int32" },
+					{ "name": "pads", "value": [ 0, 0, 0, 0 ], "type": "Int32Array" },
+					{ "name": "strides", "value": [ 1, 1 ], "type": "Int32Array"}
+				] }
+			]
+		},
+		{
+			"name": "ConvTranspose.1D",
+			"attribute_maps": [
+				{ "attributes": [
+					{ "name": "auto_pad", "value": "NOTSET", "type": "String" },
+					{ "name": "dilations", "value": [ 1 ], "type": "Int32Array" },
+					{ "name": "group", "value": 1, "type": "Int32" },
+					{ "name": "pads", "value": [ 0, 0 ], "type": "Int32Array" },
+					{ "name": "strides", "value": [ 1 ], "type": "Int32Array"}
+				] }
+			]
+		},
+		{
+			"name": "ConvTranspose.3D",
+			"attribute_maps": [
+				{ "attributes": [
+					{ "name": "auto_pad", "value": "NOTSET", "type": "String" },
+					{ "name": "dilations", "value": [ 1, 1, 1 ], "type": "Int32Array" },
+					{ "name": "group", "value": 1, "type": "Int32" },
+					{ "name": "pads", "value": [ 0, 0, 0, 0, 0, 0 ], "type": "Int32Array" },
+					{ "name": "strides", "value": [ 1, 1, 1 ], "type": "Int32Array"}
+				] }
+			]
+		},
+		{
+			"name": "ConvTranspose.Attributes",
+			"attribute_maps": [
+				{ "attributes": [
+					{ "name": "auto_pad", "value": "NOTSET", "type": "String" },
+					{ "name": "dilations", "value": [ 1, 1 ], "type": "Int32Array" },
+					{ "name": "group", "value": 1, "type": "Int32" },
+					{ "name": "output_padding", "value": [ 1, 1 ], "type": "Int32Array" },
+					{ "name": "pads", "value": [ 0, 0, 0, 0 ], "type": "Int32Array" },
+					{ "name": "strides", "value": [ 3, 2 ], "type": "Int32Array"}
+				] }
+			]
+		},
+		{
+			"name": "ConvTranspose.Autopad.SameLower",
+			"attribute_maps": [
+				{ "attributes": [
+					{ "name": "auto_pad", "value": "SAME_LOWER", "type": "String" },
+					{ "name": "strides", "value": [ 2, 2 ], "type": "Int32Array"}
+				] }
+			]
+		},
+		{
+			"name": "ConvTranspose.Autopad.SameUpper",
+			"attribute_maps": [
+				{ "attributes": [
+					{ "name": "auto_pad", "value": "SAME_UPPER", "type": "String" },
+					{ "name": "strides", "value": [ 2, 2 ], "type": "Int32Array"}
+				] }
+			]
 		}
 	]
 	
diff --git a/Engine/Plugins/Experimental/NNX/Source/NNXQA/Resources/NNXQATestDescSchema.json b/Engine/Plugins/Experimental/NNX/Source/NNXQA/Resources/NNXQATestDescSchema.json
--- a/Engine/Plugins/Experimental/NNX/Source/NNXQA/Resources/NNXQATestDescSchema.json
+++ b/Engine/Plugins/Experimental/NNX/Source/NNXQA/Resources/NNXQATestDescSchema.json
@@ -116,7 +116,7 @@
 					"description": "Name of the attribute"
 				},
 				"value": {
-					"type": "number",
+					"type": [ "number", "string", "array" ],
 					"description": "Value of the attribute"
 				},
 				"type": {
diff --git a/Engine/Plugins/Experimental/NNX/Source/NNXRuntimeRDG/Private/Hlsl/NNIRuntimeRDGConv.cpp b/Engine/Plugins/Experimental/NNX/Source/NNXRuntimeRDG/Private/Hlsl/NNIRuntimeRDGConv.cpp
--- a/Engine/Plugins/Experimental/NNX/Source/NNXRuntimeRDG/Private/Hlsl/NNIRuntimeRDGConv.cpp
+++ b/Engine/Plugins/Experimental/NNX/Source/NNXRuntimeRDG/Private/Hlsl/NNIRuntimeRDGConv.cpp
@@ -37,17 +37,18 @@ namespace UE::NNIRuntimeRDG::Private::Hlsl
 		int NumDimensions = 0;
 		bool HasBias = false;
 
-		// Hard coded parameters, until we accept them from JSON
+		EConvAutoPad AutoPad = EConvAutoPad::NOTSET;
+		TArray<int32> Dilations;
 		int32 Group = 1;
-		EConvAutoPad AutoPad = EConvAutoPad::VALID;
-		TArray<int32> Dilations = {1};
-		TArray<int32> Strides = {1};
-		TArray<int32> Pads = {0, 0};
+		TArray<int32> Pads;
+		TArray<int32> Strides;
 
 	public:
 
 		virtual bool Initialize(TConstArrayView<NNX::FTensorDesc> InputTensorDescs, TConstArrayView<NNX::FTensorDesc> OutputTensorDescs, const UE::NNECore::FAttributeMap& Attributes) override
 		{
+			using namespace UE::NNEHlslShaders::Internal;
+
 			TArray<NNX::FTensor> InputTensors;
 			TArray<NNX::FTensor> OutputTensors;
 			if (!NNX::ConvertConcreteTensorDescsToTensors(InputTensorDescs, InputTensors) ||
@@ -78,6 +79,21 @@ namespace UE::NNIRuntimeRDG::Private::Hlsl
 
 			NumDimensions = Input.GetShape().Rank() - 2;
 
+			TArray<int32> DilationsOrStridesDefault;
+			DilationsOrStridesDefault.Init(1, NumDimensions);
+
+			FConvCS::LexFromString(AutoPad, *Attributes.GetValue<FString>(TEXT("auto_pad")));
+			Dilations = Attributes.GetValueOrDefault<TArray<int32>>(TEXT("dilations"), DilationsOrStridesDefault);
+			Group = Attributes.GetValueOrDefault<int>(TEXT("group"), 1);
+			if (AutoPad == EConvAutoPad::NOTSET)
+			{
+				TArray<int32> PadsDefault;
+				PadsDefault.Init(1, 2 * NumDimensions);
+
+				Pads = Attributes.GetValueOrDefault<TArray<int32>>(TEXT("pads"), PadsDefault);
+			}
+			Strides = Attributes.GetValueOrDefault<TArray<int32>>(TEXT("strides"), DilationsOrStridesDefault);
+
 			return true;
 		}
 
diff --git a/Engine/Plugins/Experimental/NNX/Source/NNXRuntimeRDG/Private/Hlsl/NNIRuntimeRDGConvTranspose.cpp b/Engine/Plugins/Experimental/NNX/Source/NNXRuntimeRDG/Private/Hlsl/NNIRuntimeRDGConvTranspose.cpp
--- a/Engine/Plugins/Experimental/NNX/Source/NNXRuntimeRDG/Private/Hlsl/NNIRuntimeRDGConvTranspose.cpp
+++ b/Engine/Plugins/Experimental/NNX/Source/NNXRuntimeRDG/Private/Hlsl/NNIRuntimeRDGConvTranspose.cpp
@@ -36,18 +36,19 @@ namespace UE::NNIRuntimeRDG::Private::Hlsl
 		int NumDimensions = 0;
 		bool HasBias = false;
 
-		// Hard coded parameters, until we accept them from JSON
+		EConvTransposeAutoPad AutoPad = EConvTransposeAutoPad::NOTSET;
+		TArray<int32> Dilations;
 		int32 Group = 1;
-		EConvTransposeAutoPad AutoPad = EConvTransposeAutoPad::VALID;
-		TArray<int32> Dilations = {1};
-		TArray<int32> Strides = {1};
-		TArray<int32> Pads = {0, 0};
-		TArray<int32> OutputPadding = {0};
+		TArray<int32> OutputPadding;
+		TArray<int32> Pads;
+		TArray<int32> Strides;
 
 	public:
 
 		virtual bool Initialize(TConstArrayView<NNX::FTensorDesc> InputTensorDescs, TConstArrayView<NNX::FTensorDesc> OutputTensorDescs, const UE::NNECore::FAttributeMap& Attributes) override
 		{
+			using namespace UE::NNEHlslShaders::Internal;
+
 			TArray<NNX::FTensor> InputTensors;
 			TArray<NNX::FTensor> OutputTensors;
 			if (!NNX::ConvertConcreteTensorDescsToTensors(InputTensorDescs, InputTensors) ||
@@ -79,6 +80,26 @@ namespace UE::NNIRuntimeRDG::Private::Hlsl
 
 			NumDimensions = Input.GetShape().Rank() - 2;
 
+			TArray<int32> DilationsOrStridesDefault;
+			DilationsOrStridesDefault.Init(1, NumDimensions);
+
+			FConvTransposeCS::LexFromString(AutoPad, *Attributes.GetValue<FString>(TEXT("auto_pad")));
+			Dilations = Attributes.GetValueOrDefault<TArray<int32>>(TEXT("dilations"), DilationsOrStridesDefault);
+			Group = Attributes.GetValueOrDefault<int>(TEXT("group"), 1);
+			
+			TArray<int32> OutputPaddingDefault;
+			OutputPaddingDefault.Init(0, NumDimensions);
+			OutputPadding = Attributes.GetValueOrDefault<TArray<int32>>(TEXT("output_padding"), OutputPaddingDefault);
+			
+			if (AutoPad == EConvTransposeAutoPad::NOTSET)
+			{
+				TArray<int32> PadsDefault;
+				PadsDefault.Init(1, 2 * NumDimensions);
+
+				Pads = Attributes.GetValueOrDefault<TArray<int32>>(TEXT("pads"), PadsDefault);
+			}
+			Strides = Attributes.GetValueOrDefault<TArray<int32>>(TEXT("strides"), DilationsOrStridesDefault);
+
 			return true;
 		}
 
diff --git a/Engine/Plugins/Experimental/NNX/Source/NNXUtils/Private/NNXModelBuilderONNX.cpp b/Engine/Plugins/Experimental/NNX/Source/NNXUtils/Private/NNXModelBuilderONNX.cpp
--- a/Engine/Plugins/Experimental/NNX/Source/NNXUtils/Private/NNXModelBuilderONNX.cpp
+++ b/Engine/Plugins/Experimental/NNX/Source/NNXUtils/Private/NNXModelBuilderONNX.cpp
@@ -97,6 +97,8 @@ public:
 		// Validate model
 		if (res)
 		{
+			UE_LOG(LogNNX, Display, TEXT("OrtValidateModelFromMemory... (note: might abort if i.e. shapes are not correct!)"));
+
 			OrtStatusPtr status = OrtValidateModelFromMemory(Data.GetData(), Data.Num());
 
 			if (status)
@@ -218,6 +220,19 @@ public:
 			Attribute->set_type(onnx::AttributeProto::INT);
 			Attribute->set_i(Value.GetValue<int32>());
 		}
+		else if (Value.GetType() == ENNEAttributeDataType::Int32Array)
+		{
+			Attribute->set_type(onnx::AttributeProto::INTS);
+			for (int32 Val : Value.GetValue<TArray<int32>>())
+			{
+				Attribute->add_ints(Val);
+			}
+		}
+		else if (Value.GetType() == ENNEAttributeDataType::String)
+		{
+			Attribute->set_type(onnx::AttributeProto::STRING);
+			Attribute->set_s(TCHAR_TO_ANSI(*Value.GetValue<FString>()));
+		}
 		else
 		{
 			checkf(false, TEXT("not implemented"))
diff --git a/Engine/Plugins/Experimental/NNX/Source/NNXUtils/Private/NNXModelOptimizer.cpp b/Engine/Plugins/Experimental/NNX/Source/NNXUtils/Private/NNXModelOptimizer.cpp
--- a/Engine/Plugins/Experimental/NNX/Source/NNXUtils/Private/NNXModelOptimizer.cpp
+++ b/Engine/Plugins/Experimental/NNX/Source/NNXUtils/Private/NNXModelOptimizer.cpp
@@ -548,7 +548,23 @@ private:
 				}
 				else if (AttrInfo.type == Ort::GraphAttributeType::kInt)
 				{
-					Builder->AddOperatorAttribute(Op, AttrInfo.name, FNNEAttributeValue((int)AttrValue.i));
+					Builder->AddOperatorAttribute(Op, AttrInfo.name, FNNEAttributeValue((int32)AttrValue.i));
+				}
+				else if (AttrInfo.type == Ort::GraphAttributeType::kInts)
+				{
+					TArray<int32> Values;
+					Values.SetNumUninitialized(AttrValue.count);
+
+					for (int i = 0; i < AttrValue.count; i++)
+					{
+						Values[i] = (int32)AttrValue.ints[i];
+					}
+
+					Builder->AddOperatorAttribute(Op, AttrInfo.name, FNNEAttributeValue(Values));
+				}
+				else if (AttrInfo.type == Ort::GraphAttributeType::kString)
+				{
+					Builder->AddOperatorAttribute(Op, AttrInfo.name, FNNEAttributeValue(FString(AttrValue.s)));
 				}
 				else
 				{
