Store tensor data into NNX format

#rnx

diff --git a/Engine/Plugins/Experimental/NNX/Source/NNXCore/Public/NNXRuntimeFormat.h b/Engine/Plugins/Experimental/NNX/Source/NNXCore/Public/NNXRuntimeFormat.h
--- a/Engine/Plugins/Experimental/NNX/Source/NNXCore/Public/NNXRuntimeFormat.h
+++ b/Engine/Plugins/Experimental/NNX/Source/NNXCore/Public/NNXRuntimeFormat.h
@@ -68,23 +68,6 @@ struct FMLFormatOperatorDesc
 
 	UPROPERTY(VisibleAnywhere, Category = "Neural Network Inference")
 	TArray<FMLFormatAttributeDesc> Attributes;
-
-	//UPROPERTY(VisibleAnywhere, Category = "Neural Network Inference")
-	//FString InstanceName;		//!< For example "Relu1"
-
-	/*
-	UPROPERTY(VisibleAnywhere, Category = "Neural Network Inference")
-	uint32 InTensorStart;
-
-	UPROPERTY(VisibleAnywhere, Category = "Neural Network Inference")
-	uint32 InTensorCount;
-
-	UPROPERTY(VisibleAnywhere, Category = "Neural Network Inference")
-	uint32 OutTensorStart;
-
-	UPROPERTY(VisibleAnywhere, Category = "Neural Network Inference")
-	uint32 OutTensorCount;
-	*/
 };
 
 USTRUCT()
diff --git a/Engine/Plugins/Experimental/NNX/Source/NNXUtils/Private/NNXModelBuilderNNX.cpp b/Engine/Plugins/Experimental/NNX/Source/NNXUtils/Private/NNXModelBuilderNNX.cpp
--- a/Engine/Plugins/Experimental/NNX/Source/NNXUtils/Private/NNXModelBuilderNNX.cpp
+++ b/Engine/Plugins/Experimental/NNX/Source/NNXUtils/Private/NNXModelBuilderNNX.cpp
@@ -109,7 +109,7 @@ public:
 		return !Data.IsEmpty();
 	}
 
-	virtual HTensor AddTensor(const FString& Name, EMLTensorDataType DataType, TArrayView<const int32> Shape)
+	virtual HTensor AddTensor(const FString& Name, EMLTensorDataType DataType, TArrayView<const int32> Shape, const void* Data, uint64 DataSize)
 	{
 		TArray<int32, TInlineAllocator<FTensorShape::MaxRank>> NNIShape;
 		for (int i = 0; i < Shape.Num(); ++i)
@@ -118,8 +118,8 @@ public:
 			NNIShape.Emplace(Shape[i] == 0 ? -1 : Shape[i]);
 		}
 		
-		int Idx = AddTensor(Name, NNIShape, DataType, EMLFormatTensorType::None);
-		
+		int Idx = AddTensor(Name, NNIShape, DataType, Data, DataSize);
+
 		return MakeTensorHandle(reinterpret_cast<void*>((int64) Idx));
 	}
 
@@ -239,7 +239,7 @@ public:
 
 private:
 
-	int AddTensor(const FString& InName, TArrayView<const int32> InShape, EMLTensorDataType InDataType, EMLFormatTensorType InType)
+	int AddTensor(const FString& InName, TArrayView<const int32> InShape, EMLTensorDataType InDataType, const void* Data, uint64 DataSize)
 	{
 		int Idx = -1;
 
@@ -255,10 +255,21 @@ private:
 
 			Desc.Name = InName;
 			Desc.Shape = InShape;
-			Desc.Type = InType;
+			Desc.Type = EMLFormatTensorType::None;
 			Desc.DataType = InDataType;
-			Desc.DataOffset = 0;
-			Desc.DataSize = 0;
+			
+			if (Data)
+			{
+				Desc.DataOffset = Format.TensorData.AddUninitialized(DataSize);
+				Desc.DataSize = DataSize;
+
+				FMemory::Memcpy(Format.TensorData.GetData() + Desc.DataOffset, Data, DataSize);
+			}
+			else
+			{
+				Desc.DataOffset = 0;
+				Desc.DataSize = 0;
+			}
 
 			Format.Tensors.Add(Desc);
 			Idx = Format.Tensors.Num() - 1;
diff --git a/Engine/Plugins/Experimental/NNX/Source/NNXUtils/Private/NNXModelBuilderONNX.cpp b/Engine/Plugins/Experimental/NNX/Source/NNXUtils/Private/NNXModelBuilderONNX.cpp
--- a/Engine/Plugins/Experimental/NNX/Source/NNXUtils/Private/NNXModelBuilderONNX.cpp
+++ b/Engine/Plugins/Experimental/NNX/Source/NNXUtils/Private/NNXModelBuilderONNX.cpp
@@ -112,7 +112,7 @@ public:
 	}
 
 	/** Add tensor */
-	virtual HTensor AddTensor(const FString& Name, EMLTensorDataType DataType, TArrayView<const int32> Shape) override
+	virtual HTensor AddTensor(const FString& Name, EMLTensorDataType DataType, TArrayView<const int32> Shape, const void* Data, uint64 DataSize) override
 	{
 		auto Value = onnx::ValueInfoProto().New(Graph->GetArena());
 		
@@ -121,6 +121,11 @@ public:
 			return HTensor();
 		}
 
+		if (Data)
+		{
+			UE_LOG(LogNNX, Warning, TEXT("ONNX ModelBuilder currently doesn't support tensors with data"));
+		}
+
 		return MakeTensorHandle(Value);
 	}
 
diff --git a/Engine/Plugins/Experimental/NNX/Source/NNXUtils/Private/NNXModelOptimizer.cpp b/Engine/Plugins/Experimental/NNX/Source/NNXUtils/Private/NNXModelOptimizer.cpp
--- a/Engine/Plugins/Experimental/NNX/Source/NNXUtils/Private/NNXModelOptimizer.cpp
+++ b/Engine/Plugins/Experimental/NNX/Source/NNXUtils/Private/NNXModelOptimizer.cpp
@@ -530,6 +530,8 @@ private:
 		}
 
 		// Traverse all the nodes get their inputs, outputs and tensor data
+		TArray<uint8> TensorDataBlob;
+
 		for (int Idx = 0; Idx < GraphInfo.nodeCount; ++Idx)
 		{
 			Ort::GraphNode		Node = Graph->GetNode(Idx);
@@ -576,22 +578,33 @@ private:
 			for (int InIdx = 0; InIdx < NodeInfo.inputCount; ++InIdx)
 			{
 				Ort::GraphTensorInfo		TensorInfo = Graph->GetNodeInput(Node, InIdx);
-				Ort::GraphTensorInitializer	Init = Graph->GetTensorInitializer(TensorInfo.name);
+				Ort::GraphTensorInitializer	TensorInit = Graph->GetTensorInitializer(TensorInfo.name);
 				EMLTensorDataType			DataType = GetDataTypeFromGraphTensor(TensorInfo.dataType);
 				const void*					Data = nullptr;
 				uint64						DataSize = 0;
 
-				// TODO: Use tensor initializer data
-				if (!Init)
+				if (TensorInit)
 				{
-					 
+					DataSize = Graph->GetTensorDataSize(TensorInit);
+
+					if (DataSize)
+					{
+						if (TensorDataBlob.Num() < DataSize)
+						{
+							TensorDataBlob.SetNumUninitialized(DataSize);
+						}
+
+						Graph->GetTensorData(TensorInit, TensorDataBlob.GetData(), DataSize, 0);
+						Data = TensorDataBlob.GetData();
+					}
 				}
-				
+
 				auto Tensor =
 					Builder->AddTensor(
 						ANSI_TO_TCHAR(TensorInfo.name),
 						DataType,
-						MakeArrayView(TensorInfo.shape, TensorInfo.shapeLen)
+						MakeArrayView(TensorInfo.shape, TensorInfo.shapeLen),
+						Data, DataSize
 					);
 
 				Builder->AddOperatorInput(Op, Tensor);
@@ -613,14 +626,6 @@ private:
 			}
 		}
 
-		//// Traverse all tensor initializers
-		//for (int Idx = 0; Idx < GraphInfo.tensorInitializerCount; ++Idx)
-		//{
-		//	Ort::GraphTensorInitializer TensorInit = Graph->GetTensorInitializer(Idx);
-
-		//	Builder.AddTensor();
-		//}
-
 		return Builder->End(NNXData);
 	}
 };
diff --git a/Engine/Plugins/Experimental/NNX/Source/NNXUtils/Public/NNXModelBuilder.h b/Engine/Plugins/Experimental/NNX/Source/NNXUtils/Public/NNXModelBuilder.h
--- a/Engine/Plugins/Experimental/NNX/Source/NNXUtils/Public/NNXModelBuilder.h
+++ b/Engine/Plugins/Experimental/NNX/Source/NNXUtils/Public/NNXModelBuilder.h
@@ -52,7 +52,7 @@ public:
 	virtual bool End(TArray<uint8>& Data) = 0;
 
 	/** Add tensor */
-	virtual HTensor AddTensor(const FString& Name, EMLTensorDataType DataType, TArrayView<const int32> Shape) = 0;
+	virtual HTensor AddTensor(const FString& Name, EMLTensorDataType DataType, TArrayView<const int32> Shape, const void* Data = nullptr, uint64 DataSize = 0) = 0;
 
 	/** Add model input */
 	virtual bool AddInput(HTensor InTensor) = 0;
