Removed unnecessary readback wait.

diff --git a/Engine/Plugins/Experimental/NNX/Source/NNXRuntimeRDG/Private/NNXRuntimeRDG.cpp b/Engine/Plugins/Experimental/NNX/Source/NNXRuntimeRDG/Private/NNXRuntimeRDG.cpp
--- a/Engine/Plugins/Experimental/NNX/Source/NNXRuntimeRDG/Private/NNXRuntimeRDG.cpp
+++ b/Engine/Plugins/Experimental/NNX/Source/NNXRuntimeRDG/Private/NNXRuntimeRDG.cpp
@@ -275,14 +275,8 @@ int FMLInferenceModelRDG::Run(TConstArrayView<FMLTensorBinding> InInputBindings,
 				RDGBuilder.Execute();
 
 				// FIXME: Using BlockUntilGPUIdle() prevents hang on Linux
+				// FIXME: Adapt to redesigned readback API (UE 5.2)
 				RHICmdList.BlockUntilGPUIdle();
-				//RHICmdList.SubmitCommandsHint();
-
-				// Wait for readback
-				while (!Readback.RHI->IsReady())
-				{
-					FPlatformProcess::Sleep(0.001f);
-				}
 
 				// Process readback
 				{
