NNI Dynamic shape support for HLSL: tests and operators

Problem:
* HLSL operators implementations did not support shape inference and thus dynamic shapes.
* We did not had operator level dynamic shape support capability in test framework, no want to duplicate all tests explicitely.

Solution:
* Implement shape inference (exposed via op.ComputeShape())
* Update test framework to run operator tests with both static and dynamic input shapes from same input data
* Implement operator validation (bonus, not directly related to the CL but made sense to do at the same time)

Tests:
* All NNX tests pass under win64 editor, dyn shape tests disabled for DML based engine for now

diff --git a/Engine/Plugins/Experimental/NNX/Source/NNEHlslShaders/Internal/NNEHlslShadersGemmCS.h b/Engine/Plugins/Experimental/NNX/Source/NNEHlslShaders/Internal/NNEHlslShadersGemmCS.h
--- a/Engine/Plugins/Experimental/NNX/Source/NNEHlslShaders/Internal/NNEHlslShadersGemmCS.h
+++ b/Engine/Plugins/Experimental/NNX/Source/NNEHlslShaders/Internal/NNEHlslShadersGemmCS.h
@@ -79,7 +79,7 @@ namespace UE::NNEHlslShaders::Internal
 
 		static void ModifyCompilationEnvironment(const FGlobalShaderPermutationParameters& InParameters, FShaderCompilerEnvironment& OutEnvironment);
 
-		static void FillInParameters(float Alpha, float Beta, int32 TransA, int32 TransB, const NNX::FTensor& InputA, const NNX::FTensor& InputB, const NNX::FTensor& InputC, float CScalar, FParameters& Parameters);
+		static void FillInParameters(float Alpha, float Beta, int32 TransA, int32 TransB, const NNX::FTensor& InputA, const NNX::FTensor& InputB, const NNX::FTensor* InputC, float CScalar, FParameters& Parameters);
 
 		static void FillInParametersMatMul(const NNX::FTensor& InputA, const NNX::FTensor& InputB, FParameters& Parameters);
 
diff --git a/Engine/Plugins/Experimental/NNX/Source/NNEHlslShaders/Private/NNEHlslShadersGemmCS.cpp b/Engine/Plugins/Experimental/NNX/Source/NNEHlslShaders/Private/NNEHlslShadersGemmCS.cpp
--- a/Engine/Plugins/Experimental/NNX/Source/NNEHlslShaders/Private/NNEHlslShadersGemmCS.cpp
+++ b/Engine/Plugins/Experimental/NNX/Source/NNEHlslShaders/Private/NNEHlslShadersGemmCS.cpp
@@ -80,7 +80,7 @@ namespace UE::NNEHlslShaders::Internal
 		OutEnvironment.SetDefine(TEXT("MAX_NUM_STACK_DIMENSIONS"), FGemmConstants::MAX_NUM_STACK_DIMENSIONS);
 	}
 
-	void TGemmCS::FillInParameters(float Alpha, float Beta, int32 TransA, int32 TransB, const NNX::FTensor& InputA, const NNX::FTensor& InputB, const NNX::FTensor& InputC, float CScalar, TGemmCS::FParameters& Parameters)
+	void TGemmCS::FillInParameters(float Alpha, float Beta, int32 TransA, int32 TransB, const NNX::FTensor& InputA, const NNX::FTensor& InputB, const NNX::FTensor* InputC, float CScalar, TGemmCS::FParameters& Parameters)
 	{
 		check(InputA.GetShape().Rank() == 2);
 		check(InputB.GetShape().Rank() == 2);
@@ -100,8 +100,14 @@ namespace UE::NNEHlslShaders::Internal
 		Parameters.MxK = M * K;
 		Parameters.KxN = K * N;
 		Parameters.MxN = M * N;
-		Parameters.CWidth = InputC.GetShape().Rank() == 0 ? 0 : InputC.GetShape().Data[InputC.GetShape().Rank() - 1];;
-		Parameters.CHeight = InputC.GetShape().Rank() == 0 ? 0 : InputC.GetShape().Rank() == 1 ? 1 : InputC.GetShape().Data[InputC.GetShape().Rank() - 2];
+		Parameters.CWidth = 0;
+		Parameters.CHeight = 0;
+		if (InputC)
+		{
+			Parameters.CWidth = InputC->GetShape().Rank() == 0 ? 0 : InputC->GetShape().Data[InputC->GetShape().Rank() - 1];;
+			Parameters.CHeight = InputC->GetShape().Rank() == 0 ? 0 : InputC->GetShape().Rank() == 1 ? 1 : InputC->GetShape().Data[InputC->GetShape().Rank() - 2];
+		}
+		
 		Parameters.CScalar = CScalar;
 	}
 
diff --git a/Engine/Plugins/Experimental/NNX/Source/NNXQA/Private/NNXQAParametricTest.cpp b/Engine/Plugins/Experimental/NNX/Source/NNXQA/Private/NNXQAParametricTest.cpp
--- a/Engine/Plugins/Experimental/NNX/Source/NNXQA/Private/NNXQAParametricTest.cpp
+++ b/Engine/Plugins/Experimental/NNX/Source/NNXQA/Private/NNXQAParametricTest.cpp
@@ -399,6 +399,8 @@ namespace Test
 	static bool RunParametricTest(FTests::FTestSetup& TestSetup, const FString& RuntimeFilter)
 	{
 		FNNIModelRaw ONNXModel;
+		FNNIModelRaw ONNXModelVariadic;
+		ONNXModelVariadic.Format = ENNXInferenceFormat::Invalid;
 
 		if (TestSetup.IsModelTest)
 		{
@@ -415,14 +417,19 @@ namespace Test
 		else
 		{
 			// Operator test, create model in memory
-			if (!CreateONNXModelForOperator(TestSetup.TargetName, TestSetup.Inputs, TestSetup.Outputs, TestSetup.AttributeMap, ONNXModel))
+			if (!CreateONNXModelForOperator(false, TestSetup.TargetName, TestSetup.Inputs, TestSetup.Outputs, TestSetup.AttributeMap, ONNXModel))
 			{
-				UE_LOG(LogNNX, Error, TEXT("Failed to create model for test '%s'. Test ABORTED!"), *TestSetup.TargetName);
+				UE_LOG(LogNNX, Error, TEXT("Failed to create static model for test '%s'. Test ABORTED!"), *TestSetup.TargetName);
+				return false;
+			}
+			if (!CreateONNXModelForOperator(true, TestSetup.TargetName, TestSetup.Inputs, TestSetup.Outputs, TestSetup.AttributeMap, ONNXModelVariadic))
+			{
+				UE_LOG(LogNNX, Error, TEXT("Failed to create variadic model for test '%s'. Test ABORTED!"), *TestSetup.TargetName);
 				return false;
 			}
 		}
 
-		return CompareONNXModelInferenceAcrossRuntimes(ONNXModel, TestSetup, RuntimeFilter);
+		return CompareONNXModelInferenceAcrossRuntimes(ONNXModel, ONNXModelVariadic, TestSetup, RuntimeFilter);
 	}
 
 	static FString AutomationRuntimeFilter;
diff --git a/Engine/Plugins/Experimental/NNX/Source/NNXQA/Private/NNXQAUtils.cpp b/Engine/Plugins/Experimental/NNX/Source/NNXQA/Private/NNXQAUtils.cpp
--- a/Engine/Plugins/Experimental/NNX/Source/NNXQA/Private/NNXQAUtils.cpp
+++ b/Engine/Plugins/Experimental/NNX/Source/NNXQA/Private/NNXQAUtils.cpp
@@ -426,7 +426,38 @@ namespace Test
 		return ReturnValue;
 	}
 
-	bool CompareONNXModelInferenceAcrossRuntimes(const FNNIModelRaw& ONNXModel, const FTests::FTestSetup& TestSetup, const FString& RuntimeFilter)
+	bool RunTestInferenceAndCompareToRef(const FTests::FTestSetup& TestSetup, IRuntime* Runtime, const FNNIModelRaw& ONNXModel, 
+		TArrayView<TArray<char>> RefOutputMemBuffers, TArrayView<FTensor> RefOutputTensorDescs)
+	{
+		TArray<TArray<char>> OutputMemBuffers;
+		TArray<FTensor> OutputTensorDescs;
+
+		const FString& RuntimeName = Runtime->GetRuntimeName();
+
+		bool bTestSuceeded = true;
+		float AbsoluteErrorEpsilon = TestSetup.GetAbsoluteErrorEpsilonForRuntime(RuntimeName);
+		float RelativeErrorPercent = TestSetup.GetRelativeErrorPercentForRuntime(RuntimeName);
+
+		RunTestInference(ONNXModel, TestSetup, Runtime, OutputTensorDescs, OutputMemBuffers);
+		if (OutputTensorDescs.Num() == RefOutputTensorDescs.Num())
+		{
+			for (int i = 0; i < OutputTensorDescs.Num(); ++i)
+			{
+				bTestSuceeded &= VerifyTensorResult(
+					RefOutputTensorDescs[i], RefOutputMemBuffers[i],
+					OutputTensorDescs[i], OutputMemBuffers[i],
+					AbsoluteErrorEpsilon, RelativeErrorPercent);
+			}
+		}
+		else
+		{
+			UE_LOG(LogNNX, Error, TEXT("Expecting %d output tensor(s), got %d."), RefOutputTensorDescs.Num(), OutputTensorDescs.Num());
+			bTestSuceeded = false;
+		}
+		return bTestSuceeded;
+	}
+
+	bool CompareONNXModelInferenceAcrossRuntimes(const FNNIModelRaw& ONNXModel, const FNNIModelRaw& ONNXModelVariadic, const FTests::FTestSetup& TestSetup, const FString& RuntimeFilter)
 	{
 		FString CurrentPlatform = UGameplayStatics::GetPlatformName();
 		if (TestSetup.AutomationExcludedPlatform.Contains(CurrentPlatform))
@@ -482,30 +513,26 @@ namespace Test
 			}
 			else
 			{
-				TArray<TArray<char>> OutputMemBuffers;
-				TArray<FTensor> OutputTensorDescs;
-				bool bTestSuceeded = true;
-				float AbsoluteErrorEpsilon = TestSetup.GetAbsoluteErrorEpsilonForRuntime(RuntimeName);
-				float RelativeErrorPercent = TestSetup.GetRelativeErrorPercentForRuntime(RuntimeName);
-
-				RunTestInference(ONNXModel, TestSetup, Runtime, OutputTensorDescs, OutputMemBuffers);
-				if (OutputTensorDescs.Num() == RefOutputTensorDescs.Num())
+				bool bShouldRunVariadicTest = (ONNXModelVariadic.Format != ENNXInferenceFormat::Invalid);
+				bShouldRunVariadicTest &= !(RuntimeName == "NNXRuntimeDML");
+				bShouldRunVariadicTest &= !(RuntimeName == "NNXRuntimeORTDml");
+
+				bool bTestSuceeded = RunTestInferenceAndCompareToRef(TestSetup, Runtime, ONNXModel, RefOutputMemBuffers, RefOutputTensorDescs);
+				
+				if (bShouldRunVariadicTest)
 				{
-					for (int i = 0; i < OutputTensorDescs.Num(); ++i)
+					if (!bTestSuceeded)
 					{
-						bTestSuceeded &= VerifyTensorResult(
-							RefOutputTensorDescs[i], RefOutputMemBuffers[i],
-							OutputTensorDescs[i], OutputMemBuffers[i],
-							AbsoluteErrorEpsilon, RelativeErrorPercent);
+						UE_LOG(LogNNX, Error, TEXT("Failed running static test."));
+					}
+					bool bVariadicTestSuceeded = RunTestInferenceAndCompareToRef(TestSetup, Runtime, ONNXModelVariadic, RefOutputMemBuffers, RefOutputTensorDescs);
+					if (!bVariadicTestSuceeded)
+					{
+						bTestSuceeded = false;
+						UE_LOG(LogNNX, Error, TEXT("Failed running variadic test."));
 					}
-				}
-				else
-				{
-					UE_LOG(LogNNX, Error, TEXT("Expecting %d output tensor(s), got %d."), RefOutputTensorDescs.Num(), OutputTensorDescs.Num());
-					bTestSuceeded = false;
 				}
 				TestResult = bTestSuceeded ? TEXT("SUCCESS") : TEXT("FAILED");
-				
 				bAllTestsSucceeded &= bTestSuceeded;
 			}
 
diff --git a/Engine/Plugins/Experimental/NNX/Source/NNXQA/Private/NNXQAUtils.h b/Engine/Plugins/Experimental/NNX/Source/NNXQA/Private/NNXQAUtils.h
--- a/Engine/Plugins/Experimental/NNX/Source/NNXQA/Private/NNXQAUtils.h
+++ b/Engine/Plugins/Experimental/NNX/Source/NNXQA/Private/NNXQAUtils.h
@@ -61,7 +61,7 @@ namespace Test
 		TArray<FTestSetup> TestSetups;
 	};
 
-	bool CompareONNXModelInferenceAcrossRuntimes(const FNNIModelRaw& ONNXModel, const FTests::FTestSetup& TestSetup, const FString& RuntimeFilter = TEXT(""));
+	bool CompareONNXModelInferenceAcrossRuntimes(const FNNIModelRaw& ONNXModel, const FNNIModelRaw& ONNXModelVariadic, const FTests::FTestSetup& TestSetup, const FString& RuntimeFilter = TEXT(""));
 
 	FString FMLTensorDescToString(const FTensor& TensorDesc);
 	FString TensorToString(const FTensor& TensorDesc, const TArray<char>& TensorData);
diff --git a/Engine/Plugins/Experimental/NNX/Source/NNXRuntimeRDG/Private/Hlsl/NNIRuntimeRDGConv.cpp b/Engine/Plugins/Experimental/NNX/Source/NNXRuntimeRDG/Private/Hlsl/NNIRuntimeRDGConv.cpp
--- a/Engine/Plugins/Experimental/NNX/Source/NNXRuntimeRDG/Private/Hlsl/NNIRuntimeRDGConv.cpp
+++ b/Engine/Plugins/Experimental/NNX/Source/NNXRuntimeRDG/Private/Hlsl/NNIRuntimeRDGConv.cpp
@@ -29,13 +29,7 @@ namespace UE::NNIRuntimeRDG::Private::Hlsl
 
 		FConv() {}
 
-		NNX::FTensor Input = {};
-		NNX::FTensor Weights = {};
-		NNX::FTensor Bias = {};
-		NNX::FTensor Output = {};
-
-		int NumDimensions = 0;
-		bool HasBias = false;
+		int32 NumDimensions = 0;
 
 		EConvAutoPad AutoPad = EConvAutoPad::NOTSET;
 		TArray<int32> Dilations;
@@ -45,36 +39,50 @@ namespace UE::NNIRuntimeRDG::Private::Hlsl
 
 	public:
 
-		virtual bool Initialize(TConstArrayView<NNX::FTensorDesc> InputTensorDescs, TConstArrayView<NNX::FTensorDesc> OutputTensorDescs, const UE::NNECore::FAttributeMap& Attributes) override
+		virtual int ComputeOutputShape(TConstArrayView<NNX::FTensorShape> InputShapes, TArray<NNX::FTensorShape>& OutputShapes) const override
 		{
-			using namespace UE::NNEHlslShaders::Internal;
+			OutputShapes.Empty();
+			check(InputShapes.Num() >= 2 && InputShapes.Num() <= 3);
 
-			TArray<NNX::FTensor> InputTensors;
-			TArray<NNX::FTensor> OutputTensors;
-			if (!NNX::ConvertConcreteTensorDescsToTensors(InputTensorDescs, InputTensors) ||
-				!NNX::ConvertConcreteTensorDescsToTensors(OutputTensorDescs, OutputTensors))
+			const NNX::FTensorShape& Input = InputShapes[0];
+			const NNX::FTensorShape& Weights = InputShapes[1];
+			NNX::FSymbolicTensorShape OutputShape;
+			
+			OutputShape.Data = UE::NNEHlslShaders::Internal::FConvCS::GetOutputShape(Input.Data, Weights.Data, AutoPad, Dilations, Strides, Pads);
+			if (!OutputShape.IsConcrete())
 			{
-				UE_LOG(LogNNX, Warning, TEXT("Variable input shapes are not supported by this operator"));
-				return false;
+				return -1;
 			}
-			
-			check(InputTensors.Num() >= 2 && InputTensors.Num() <= 3);
-			check(OutputTensors.Num() == 1);
+			OutputShapes.Emplace(NNX::FTensorShape::MakeFromSymbolic(OutputShape));
 
-			Input = InputTensors[0];
-			Weights = InputTensors[1];
-			Output = OutputTensors[0];
+			return 0;
+		};
 
-			check(Input.GetShape().Rank() > 2);
-			check(Weights.GetShape().Rank()  == Input.GetShape().Rank());
-			check(Output.GetShape().Rank() == Input.GetShape().Rank());
+		virtual bool Initialize(TConstArrayView<NNX::FTensorDesc> InputTensorDescs, TConstArrayView<NNX::FTensorDesc> OutputTensorDescs, const UE::NNECore::FAttributeMap& Attributes) override
+		{
+            using namespace UE::NNEHlslShaders::Internal;
 
-			if (InputTensors.Num() == 3) {
-				HasBias = true;
-				Bias = InputTensors[2];
+			check(InputTensorDescs.Num() >= 2 && InputTensorDescs.Num() <= 3);
+			check(OutputTensorDescs.Num() == 1);
+
+            const NNX::FTensorDesc& Input = InputTensorDescs[0];
+			const NNX::FTensorDesc& Weights = InputTensorDescs[1];
+			const NNX::FTensorDesc& Output = OutputTensorDescs[0];
+			
+			if (Input.GetShape().Rank() < 2)
+			{
+				UE_LOG(LogNNX, Warning, TEXT("Conv first input should be at least of rank 2"));
+				return false;
 			}
-			else {
-				HasBias = false;
+			if (Weights.GetShape().Rank() != Input.GetShape().Rank())
+			{
+				UE_LOG(LogNNX, Warning, TEXT("Conv first and second inputs should be of same ranks"));
+				return false;
+			}
+			if (Output.GetShape().Rank() != Input.GetShape().Rank())
+			{
+				UE_LOG(LogNNX, Warning, TEXT("Conv first and output should be of same ranks"));
+				return false;
 			}
 
 			NumDimensions = Input.GetShape().Rank() - 2;
@@ -97,24 +105,43 @@ namespace UE::NNIRuntimeRDG::Private::Hlsl
 			return true;
 		}
 
-		virtual void Dispatch(FRDGBuilder& GraphBuilder, TConstArrayView<NNX::FTensorRDG> InInputTensors, TConstArrayView<NNX::FTensorRDG> InOutputTensors) override
+		virtual void Dispatch(FRDGBuilder& GraphBuilder, TConstArrayView<NNX::FTensorRDG> InputTensors, TConstArrayView<NNX::FTensorRDG> OutputTensors) override
 		{
 			using namespace UE::NNEHlslShaders::Internal;
 
 			constexpr EConvAlgorithm Algorithm = EConvAlgorithm::SharedMemory;
 			constexpr EConvGroupSize GroupSize = EConvGroupSize::Size256;
 
+			check(InputTensors.Num() >= 2 && InputTensors.Num() <= 3);
+			check(OutputTensors.Num() == 1);
+
+			const NNX::FTensorRDG& Input = InputTensors[0];
+			const NNX::FTensorRDG& Weights = InputTensors[1];
+			const NNX::FTensorRDG& Output = OutputTensors[0];
+			const NNX::FTensorRDG* Bias = nullptr;
+			bool HasBias = false;
+
+			if (InputTensors.Num() == 3) {
+				HasBias = true;
+				Bias = &(InputTensors[2]);
+			}
+
+			check(Input.GetShape().Rank() > 2);
+			check(Weights.GetShape().Rank() == Input.GetShape().Rank());
+			check(Output.GetShape().Rank() == Input.GetShape().Rank());
+			check(NumDimensions == (Input.GetShape().Rank() - 2));
+
 			TArray<int32> OutputShape = FConvCS::GetOutputShape(Input.GetShape().Data, Weights.GetShape().Data, AutoPad, Dilations, Strides, Pads);
 
 			// Set parameters
 			FConvCS::FParameters* Params = GraphBuilder.AllocParameters<FConvCS::FParameters>();
 			FConvCS::FillInParameters(GroupSize, Input.GetShape().Data, Weights.GetShape().Data, HasBias, AutoPad, Group, Dilations,Strides, Pads, *Params);
-			Params->X = GraphBuilder.CreateSRV(FRDGBufferSRVDesc(InInputTensors[0].GetBuffer(), PF_R32_FLOAT));
-			Params->W = GraphBuilder.CreateSRV(FRDGBufferSRVDesc(InInputTensors[1].GetBuffer(), PF_R32_FLOAT));
-			if (InInputTensors.Num() == 3) {
-				Params->B = GraphBuilder.CreateSRV(FRDGBufferSRVDesc(InInputTensors[2].GetBuffer(), PF_R32_FLOAT));
+			Params->X = GraphBuilder.CreateSRV(FRDGBufferSRVDesc(Input.GetBuffer(), PF_R32_FLOAT));
+			Params->W = GraphBuilder.CreateSRV(FRDGBufferSRVDesc(Weights.GetBuffer(), PF_R32_FLOAT));
+			if (HasBias) {
+				Params->B = GraphBuilder.CreateSRV(FRDGBufferSRVDesc(Bias->GetBuffer(), PF_R32_FLOAT));
 			}
-			Params->Y = GraphBuilder.CreateUAV(FRDGBufferUAVDesc(InOutputTensors[0].GetBuffer(), PF_R32_FLOAT));
+			Params->Y = GraphBuilder.CreateUAV(FRDGBufferUAVDesc(Output.GetBuffer(), PF_R32_FLOAT));
 
 			FConvCS::FPermutationDomain PermutationVector;
 
@@ -138,9 +165,33 @@ namespace UE::NNIRuntimeRDG::Private::Hlsl
 		}
 	};
 
+	bool ValidateConvOperator(const UE::NNECore::FAttributeMap& AttributeMap, TConstArrayView<EMLTensorDataType> InputTypes, TConstArrayView<NNX::FSymbolicTensorShape> InputShapes)
+	{
+		bool bIsValid = true;
+
+		NNX::FAttributeValidator AttributeValidator;
+		AttributeValidator.AddOptional(TEXT("auto_pad"), ENNEAttributeDataType::String);
+		AttributeValidator.AddOptional(TEXT("dilations"), ENNEAttributeDataType::Int32Array);
+		AttributeValidator.AddOptional(TEXT("group"), ENNEAttributeDataType::Int32);
+		//AttributeValidator.AddOptional(TEXT("kernel_shape"), ENNEAttributeDataType::Int32Array);
+		AttributeValidator.AddOptional(TEXT("pads"), ENNEAttributeDataType::Int32Array);
+		AttributeValidator.AddOptional(TEXT("strides"), ENNEAttributeDataType::Int32Array);
+
+		bIsValid &= AttributeValidator.Validate(AttributeMap);
+
+		NNX::FInputValidator InputValidator;
+		InputValidator.AddSupportedType(EMLTensorDataType::Float);
+		InputValidator.AddRequired();
+		InputValidator.AddRequired();
+		InputValidator.AddOptional();
+		bIsValid &= InputValidator.Validate(InputTypes);
+
+		return bIsValid;
+	}
+
 	bool RegisterConvOperator(NNX::FMLOperatorRegistryHlsl& Registry)
 	{
-		Registry.OpAdd(TEXT("Conv"), FConv::Create);
+		Registry.OpAdd(TEXT("Conv"), FConv::Create, ValidateConvOperator);
 
 		return true;
 	}
diff --git a/Engine/Plugins/Experimental/NNX/Source/NNXRuntimeRDG/Private/Hlsl/NNIRuntimeRDGConvTranspose.cpp b/Engine/Plugins/Experimental/NNX/Source/NNXRuntimeRDG/Private/Hlsl/NNIRuntimeRDGConvTranspose.cpp
--- a/Engine/Plugins/Experimental/NNX/Source/NNXRuntimeRDG/Private/Hlsl/NNIRuntimeRDGConvTranspose.cpp
+++ b/Engine/Plugins/Experimental/NNX/Source/NNXRuntimeRDG/Private/Hlsl/NNIRuntimeRDGConvTranspose.cpp
@@ -28,13 +28,7 @@ namespace UE::NNIRuntimeRDG::Private::Hlsl
 
 		FConvTranspose() {}
 
-		NNX::FTensor Input = {};
-		NNX::FTensor Weights = {};
-		NNX::FTensor Bias = {};
-		NNX::FTensor Output = {};
-
-		int NumDimensions = 0;
-		bool HasBias = false;
+		int32 NumDimensions = 0;
 
 		EConvTransposeAutoPad AutoPad = EConvTransposeAutoPad::NOTSET;
 		TArray<int32> Dilations;
@@ -45,37 +39,50 @@ namespace UE::NNIRuntimeRDG::Private::Hlsl
 
 	public:
 
+		virtual int ComputeOutputShape(TConstArrayView<NNX::FTensorShape> InputShapes, TArray<NNX::FTensorShape>& OutputShapes) const override
+		{
+			OutputShapes.Empty();
+			check(InputShapes.Num() >= 2 && InputShapes.Num() <= 3);
+			
+			const NNX::FTensorShape& Input = InputShapes[0];
+			const NNX::FTensorShape& Weights = InputShapes[1];
+			NNX::FSymbolicTensorShape OutputShape;
+			
+			OutputShape.Data = UE::NNEHlslShaders::Internal::FConvTransposeCS::GetOutputShape(Input.Data, Weights.Data, AutoPad, Dilations, Strides, Pads, OutputPadding, Group);
+			if (!OutputShape.IsConcrete())
+			{
+				return -1;
+			}
+			OutputShapes.Emplace(NNX::FTensorShape::MakeFromSymbolic(OutputShape));
+
+			return 0;
+		};
+
 		virtual bool Initialize(TConstArrayView<NNX::FTensorDesc> InputTensorDescs, TConstArrayView<NNX::FTensorDesc> OutputTensorDescs, const UE::NNECore::FAttributeMap& Attributes) override
 		{
 			using namespace UE::NNEHlslShaders::Internal;
 
-			TArray<NNX::FTensor> InputTensors;
-			TArray<NNX::FTensor> OutputTensors;
-			if (!NNX::ConvertConcreteTensorDescsToTensors(InputTensorDescs, InputTensors) ||
-				!NNX::ConvertConcreteTensorDescsToTensors(OutputTensorDescs, OutputTensors))
+			check(InputTensorDescs.Num() >= 2 && InputTensorDescs.Num() <= 3);
+			check(OutputTensorDescs.Num() == 1);
+
+            const NNX::FTensorDesc& Input = InputTensorDescs[0];
+			const NNX::FTensorDesc& Weights = InputTensorDescs[1];
+			const NNX::FTensorDesc& Output = OutputTensorDescs[0];
+
+			if (Input.GetShape().Rank() < 2)
 			{
-				UE_LOG(LogNNX, Warning, TEXT("Variable input shapes are not supported by this operator"));
+				UE_LOG(LogNNX, Warning, TEXT("ConvTranspose first input should be at least of rank 2"));
 				return false;
 			}
-				
-			check(InputTensors.Num() >= 2 && InputTensors.Num() <= 3);
-			check(OutputTensors.Num() == 1);
-
-			Input = InputTensors[0];
-			Weights = InputTensors[1];
-			Output = OutputTensors[0];
-
-			check(Input.GetShape().Rank() > 2);
-			check(Weights.GetShape().Rank()  == Input.GetShape().Rank());
-			check(Output.GetShape().Rank() == Input.GetShape().Rank());
-
-			if (InputTensors.Num() == 3) {
-				HasBias = true;
-				Bias = InputTensors[2];
+			if (Weights.GetShape().Rank() != Input.GetShape().Rank())
+			{
+				UE_LOG(LogNNX, Warning, TEXT("ConvTranspose first and second inputs should be of same ranks"));
+				return false;
 			}
-			else
+			if (Output.GetShape().Rank() != Input.GetShape().Rank())
 			{
-				HasBias = false;
+				UE_LOG(LogNNX, Warning, TEXT("ConvTranspose first input and output should be of same ranks"));
+				return false;
 			}
 
 			NumDimensions = Input.GetShape().Rank() - 2;
@@ -103,24 +110,43 @@ namespace UE::NNIRuntimeRDG::Private::Hlsl
 			return true;
 		}
 
-		virtual void Dispatch(FRDGBuilder& GraphBuilder, TConstArrayView<NNX::FTensorRDG> InInputTensors, TConstArrayView<NNX::FTensorRDG> InOutputTensors) override
+		virtual void Dispatch(FRDGBuilder& GraphBuilder, TConstArrayView<NNX::FTensorRDG> InputTensors, TConstArrayView<NNX::FTensorRDG> OutputTensors) override
 		{
 			using namespace UE::NNEHlslShaders::Internal;
 
 			constexpr EConvTransposeAlgorithm Algorithm = EConvTransposeAlgorithm::SharedMemory;
 			constexpr EConvTransposeGroupSize GroupSize = EConvTransposeGroupSize::Size256;
 
+			check(InputTensors.Num() >= 2 && InputTensors.Num() <= 3);
+			check(OutputTensors.Num() == 1);
+
+			const NNX::FTensorRDG& Input = InputTensors[0];
+			const NNX::FTensorRDG& Weights = InputTensors[1];
+			const NNX::FTensorRDG& Output = OutputTensors[0];
+			const NNX::FTensorRDG* Bias = nullptr;
+			bool HasBias = false;
+
+			if (InputTensors.Num() == 3) {
+				HasBias = true;
+				Bias = &(InputTensors[2]);
+			}
+
+			check(Input.GetShape().Rank() > 2);
+			check(Weights.GetShape().Rank() == Input.GetShape().Rank());
+			check(Output.GetShape().Rank() == Input.GetShape().Rank());
+			check(NumDimensions == (Input.GetShape().Rank() - 2));
+
 			TArray<int32> OutputShape = FConvTransposeCS::GetOutputShape(Input.GetShape().Data, Weights.GetShape().Data, AutoPad, Dilations, Strides, Pads, OutputPadding, Group);
 
 			// Set parameters
 			FConvTransposeCS::FParameters* Params = GraphBuilder.AllocParameters<FConvTransposeCS::FParameters>();
 			FConvTransposeCS::FillInParameters(GroupSize, Input.GetShape().Data, Weights.GetShape().Data, HasBias, AutoPad, Group, Dilations,Strides, Pads, OutputPadding, *Params);
-			Params->X = GraphBuilder.CreateSRV(FRDGBufferSRVDesc(InInputTensors[0].GetBuffer(), PF_R32_FLOAT));
-			Params->W = GraphBuilder.CreateSRV(FRDGBufferSRVDesc(InInputTensors[1].GetBuffer(), PF_R32_FLOAT));
-			if (InInputTensors.Num() == 3) {
-				Params->B = GraphBuilder.CreateSRV(FRDGBufferSRVDesc(InInputTensors[2].GetBuffer(), PF_R32_FLOAT));
+			Params->X = GraphBuilder.CreateSRV(FRDGBufferSRVDesc(Input.GetBuffer(), PF_R32_FLOAT));
+			Params->W = GraphBuilder.CreateSRV(FRDGBufferSRVDesc(Weights.GetBuffer(), PF_R32_FLOAT));
+			if (HasBias) {
+				Params->B = GraphBuilder.CreateSRV(FRDGBufferSRVDesc(Bias->GetBuffer(), PF_R32_FLOAT));
 			}
-			Params->Y = GraphBuilder.CreateUAV(FRDGBufferUAVDesc(InOutputTensors[0].GetBuffer(), PF_R32_FLOAT));
+			Params->Y = GraphBuilder.CreateUAV(FRDGBufferUAVDesc(Output.GetBuffer(), PF_R32_FLOAT));
 
 			FConvTransposeCS::FPermutationDomain PermutationVector;
 
@@ -144,9 +170,35 @@ namespace UE::NNIRuntimeRDG::Private::Hlsl
 		}
 	};
 
+	bool ValidateConvTransposeOperator(const UE::NNECore::FAttributeMap& AttributeMap, TConstArrayView<EMLTensorDataType> InputTypes, TConstArrayView<NNX::FSymbolicTensorShape> InputShapes)
+	{
+		bool bIsValid = true;
+
+		NNX::FAttributeValidator AttributeValidator;
+		AttributeValidator.AddOptional(TEXT("auto_pad"), ENNEAttributeDataType::String);
+		AttributeValidator.AddOptional(TEXT("dilations"), ENNEAttributeDataType::Int32Array);
+		AttributeValidator.AddOptional(TEXT("group"), ENNEAttributeDataType::Int32);
+		//AttributeValidator.AddOptional(TEXT("kernel_shape"), ENNEAttributeDataType::Int32Array);
+		AttributeValidator.AddOptional(TEXT("output_padding"), ENNEAttributeDataType::Int32Array);
+		//AttributeValidator.AddOptional(TEXT("output_shape"), ENNEAttributeDataType::Int32Array);
+		AttributeValidator.AddOptional(TEXT("pads"), ENNEAttributeDataType::Int32Array);
+		AttributeValidator.AddOptional(TEXT("strides"), ENNEAttributeDataType::Int32Array);
+
+		bIsValid &= AttributeValidator.Validate(AttributeMap);
+
+		NNX::FInputValidator InputValidator;
+		InputValidator.AddSupportedType(EMLTensorDataType::Float);
+		InputValidator.AddRequired();
+		InputValidator.AddRequired();
+		InputValidator.AddOptional();
+		bIsValid &= InputValidator.Validate(InputTypes);
+
+		return bIsValid;
+	}
+
 	bool RegisterConvTransposeOperator(NNX::FMLOperatorRegistryHlsl& Registry)
 	{
-		Registry.OpAdd(TEXT("ConvTranspose"), FConvTranspose::Create);
+		Registry.OpAdd(TEXT("ConvTranspose"), FConvTranspose::Create, ValidateConvTransposeOperator);
 
 		return true;
 	}
diff --git a/Engine/Plugins/Experimental/NNX/Source/NNXRuntimeRDG/Private/Hlsl/NNIRuntimeRDGElementWiseBinary.cpp b/Engine/Plugins/Experimental/NNX/Source/NNXRuntimeRDG/Private/Hlsl/NNIRuntimeRDGElementWiseBinary.cpp
--- a/Engine/Plugins/Experimental/NNX/Source/NNXRuntimeRDG/Private/Hlsl/NNIRuntimeRDGElementWiseBinary.cpp
+++ b/Engine/Plugins/Experimental/NNX/Source/NNXRuntimeRDG/Private/Hlsl/NNIRuntimeRDGElementWiseBinary.cpp
@@ -22,40 +22,60 @@ namespace UE::NNIRuntimeRDG::Private::Hlsl
 		TElementWiseBinary() {}
 		virtual ~TElementWiseBinary() = default;
 
-	private:
-
-		NNX::FTensor LHSInput = {};
-		NNX::FTensor RHSInput = {};
-		NNX::FTensor Output = {};
-
 	public:
 
-		virtual bool Initialize(TConstArrayView<NNX::FTensorDesc> InputTensorDescs, TConstArrayView<NNX::FTensorDesc> OutputTensorDescs, const UE::NNECore::FAttributeMap& Attributes) override
+		virtual int ComputeOutputShape(TConstArrayView<NNX::FTensorShape> InputShapes, TArray<NNX::FTensorShape>& OutputShapes) const override
 		{
-			TArray<NNX::FTensor> InputTensors;
-			TArray<NNX::FTensor> OutputTensors;
-			if (!NNX::ConvertConcreteTensorDescsToTensors(InputTensorDescs, InputTensors) ||
-				!NNX::ConvertConcreteTensorDescsToTensors(OutputTensorDescs, OutputTensors))
+			OutputShapes.Empty();
+			check(InputShapes.Num() == 2);
+			
+			const NNX::FTensorShape& LHSInput = InputShapes[0];
+			const NNX::FTensorShape& RHSInput = InputShapes[1];
+			const int32 OutputRank = FMath::Max(LHSInput.Rank(), RHSInput.Rank());
+			NNX::FTensorShape OutputShape;
+			
+			OutputShape.Data.SetNumUninitialized(OutputRank);
+			
+			for (int32 i = 0; i < OutputRank; ++i)
 			{
-				UE_LOG(LogNNX, Warning, TEXT("Variable input shapes are not supported by this operator"));
-				return false;
+				int32 LHSIndex = LHSInput.Rank() - 1 - i;
+				int32 RHSIndex = RHSInput.Rank() - 1 - i;
+				int32 LHSValue = LHSIndex >= 0 ? LHSInput.Data[LHSIndex] : 1;
+				int32 RHSValue = RHSIndex >= 0 ? RHSInput.Data[RHSIndex] : 1;
+				if (LHSValue != RHSValue && LHSValue != 1 && RHSValue != 1)
+				{
+					UE_LOG(LogNNX, Warning, TEXT("Error while computing shape for element wise binary op, input shapes are not compatible"));
+					return -1;
+				}
+				int32 OutputValue = FMath::Max(LHSValue, RHSValue);
+				OutputShape.Data[OutputRank - 1 - i] = OutputValue;
 			}
+
+			OutputShapes.Emplace(OutputShape);
 			
-			check(InputTensors.Num() == 2);
-			check(OutputTensors.Num() == 1);
+			return 0;
+		}
+		
+		virtual bool Initialize(TConstArrayView<NNX::FTensorDesc> InputTensorDescs, TConstArrayView<NNX::FTensorDesc> OutputTensorDescs, const UE::NNECore::FAttributeMap& Attributes) override
+		{
+			check(InputTensorDescs.Num() == 2);
+			check(OutputTensorDescs.Num() == 1);
 		
-			LHSInput = InputTensors[0];
-			RHSInput = InputTensors[1];
-			Output = OutputTensors[0];
-
 			return true;
 		}
 
-		virtual void Dispatch(FRDGBuilder& GraphBuilder, TConstArrayView<NNX::FTensorRDG> InInputTensors, TConstArrayView<NNX::FTensorRDG> InOutputTensors) override
+		virtual void Dispatch(FRDGBuilder& GraphBuilder, TConstArrayView<NNX::FTensorRDG> InputTensors, TConstArrayView<NNX::FTensorRDG> OutputTensors) override
 		{
-			FRDGBufferSRVRef LHSInputSRV = GraphBuilder.CreateSRV(FRDGBufferSRVDesc(InInputTensors[0].GetBuffer(), PF_R32_FLOAT));
-			FRDGBufferSRVRef RHSInputSRV = GraphBuilder.CreateSRV(FRDGBufferSRVDesc(InInputTensors[1].GetBuffer(), PF_R32_FLOAT));
-			FRDGBufferUAVRef OutputUAV = GraphBuilder.CreateUAV(FRDGBufferUAVDesc(InOutputTensors[0].GetBuffer(), PF_R32_FLOAT));
+			check(InputTensors.Num() == 2);
+			check(OutputTensors.Num() == 1);
+
+			const NNX::FTensorRDG& LHSInput = InputTensors[0];
+			const NNX::FTensorRDG& RHSInput = InputTensors[1];
+			const NNX::FTensorRDG& Output = OutputTensors[0];
+			
+			FRDGBufferSRVRef LHSInputSRV = GraphBuilder.CreateSRV(FRDGBufferSRVDesc(LHSInput.GetBuffer(), PF_R32_FLOAT));
+			FRDGBufferSRVRef RHSInputSRV = GraphBuilder.CreateSRV(FRDGBufferSRVDesc(RHSInput.GetBuffer(), PF_R32_FLOAT));
+			FRDGBufferUAVRef OutputUAV = GraphBuilder.CreateUAV(FRDGBufferUAVDesc(Output.GetBuffer(), PF_R32_FLOAT));
 
 			FIntVector ThreadGroupCount = NNX::ComputeElementWiseThreadGroups(Output.GetVolume(), FElementWiseBinaryConstants::NUM_GROUP_THREADS);
 
@@ -90,6 +110,22 @@ namespace UE::NNIRuntimeRDG::Private::Hlsl
 		}
 	};
 
+	bool ValidateElementWiseBinaryOperator(const UE::NNECore::FAttributeMap& AttributeMap, TConstArrayView<EMLTensorDataType> InputTypes, TConstArrayView<NNX::FSymbolicTensorShape> InputShapes)
+	{
+		bool bIsValid = true;
+
+		NNX::FAttributeValidator AttributeValidator;
+		bIsValid &= AttributeValidator.Validate(AttributeMap);
+
+		NNX::FInputValidator InputValidator;
+		InputValidator.AddSupportedType(EMLTensorDataType::Float);
+		InputValidator.AddRequired();
+		InputValidator.AddRequired();
+		bIsValid &= InputValidator.Validate(InputTypes);
+
+		return bIsValid;
+	}
+
 	template<EMLElementWiseBinaryOperatorType OpType>
 	NNX::FMLOperatorHlsl* CreateElementWiseBinaryOperator()
 	{
@@ -98,7 +134,7 @@ namespace UE::NNIRuntimeRDG::Private::Hlsl
 
 	bool RegisterElementWiseBinaryOperators(NNX::FMLOperatorRegistryHlsl& Registry)
 	{
-#define OP(Name) Registry.OpAdd(TEXT(#Name), CreateElementWiseBinaryOperator<EMLElementWiseBinaryOperatorType::Name>)
+#define OP(Name) Registry.OpAdd(TEXT(#Name), CreateElementWiseBinaryOperator<EMLElementWiseBinaryOperatorType::Name>, ValidateElementWiseBinaryOperator)
 		OP(Add);
 		//OP(And);
 		OP(Div);
diff --git a/Engine/Plugins/Experimental/NNX/Source/NNXRuntimeRDG/Private/Hlsl/NNIRuntimeRDGElementWiseUnary.cpp b/Engine/Plugins/Experimental/NNX/Source/NNXRuntimeRDG/Private/Hlsl/NNIRuntimeRDGElementWiseUnary.cpp
--- a/Engine/Plugins/Experimental/NNX/Source/NNXRuntimeRDG/Private/Hlsl/NNIRuntimeRDGElementWiseUnary.cpp
+++ b/Engine/Plugins/Experimental/NNX/Source/NNXRuntimeRDG/Private/Hlsl/NNIRuntimeRDGElementWiseUnary.cpp
@@ -112,9 +112,95 @@ namespace UE::NNIRuntimeRDG::Private::Hlsl
 		return new TElementWiseUnary<OpType>();
 	}
 
+	template<EMLElementWiseUnaryOperatorType OpType>
+	bool ValidateElementWiseUnaryOperator(const UE::NNECore::FAttributeMap& AttributeMap, TConstArrayView<EMLTensorDataType> InputTypes, TConstArrayView<NNX::FSymbolicTensorShape> InputShapes)
+	{
+		bool bIsValid = true;
+
+		NNX::FAttributeValidator AttributeValidator;
+		bIsValid &= AttributeValidator.Validate(AttributeMap);
+
+		NNX::FInputValidator InputValidator;
+		InputValidator.AddSupportedType(EMLTensorDataType::Float);
+		InputValidator.AddRequired();
+		bIsValid &= InputValidator.Validate(InputTypes);
+
+		return bIsValid;
+	}
+
+	template<>
+	bool ValidateElementWiseUnaryOperator<EMLElementWiseUnaryOperatorType::Selu>(const UE::NNECore::FAttributeMap& AttributeMap, TConstArrayView<EMLTensorDataType> InputTypes, TConstArrayView<NNX::FSymbolicTensorShape> InputShapes)
+	{
+		bool bIsValid = true;
+
+		NNX::FAttributeValidator AttributeValidator;
+		AttributeValidator.AddOptional(TEXT("alpha"), ENNEAttributeDataType::Float);
+		AttributeValidator.AddOptional(TEXT("gamma"), ENNEAttributeDataType::Float);
+		bIsValid &= AttributeValidator.Validate(AttributeMap);
+
+		NNX::FInputValidator InputValidator;
+		InputValidator.AddSupportedType(EMLTensorDataType::Float);
+		InputValidator.AddRequired();
+		bIsValid &= InputValidator.Validate(InputTypes);
+
+		return bIsValid;
+	}
+
+	template<>
+	bool ValidateElementWiseUnaryOperator<EMLElementWiseUnaryOperatorType::Elu>(const UE::NNECore::FAttributeMap& AttributeMap, TConstArrayView<EMLTensorDataType> InputTypes, TConstArrayView<NNX::FSymbolicTensorShape> InputShapes)
+	{
+		bool bIsValid = true;
+
+		NNX::FAttributeValidator AttributeValidator;
+		AttributeValidator.AddOptional(TEXT("alpha"), ENNEAttributeDataType::Float);
+		bIsValid &= AttributeValidator.Validate(AttributeMap);
+
+		NNX::FInputValidator InputValidator;
+		InputValidator.AddSupportedType(EMLTensorDataType::Float);
+		InputValidator.AddRequired();
+		bIsValid &= InputValidator.Validate(InputTypes);
+
+		return bIsValid;
+	}
+
+	template<>
+	bool ValidateElementWiseUnaryOperator<EMLElementWiseUnaryOperatorType::HardSigmoid>(const UE::NNECore::FAttributeMap& AttributeMap, TConstArrayView<EMLTensorDataType> InputTypes, TConstArrayView<NNX::FSymbolicTensorShape> InputShapes)
+	{
+		bool bIsValid = true;
+
+		NNX::FAttributeValidator AttributeValidator;
+		AttributeValidator.AddOptional(TEXT("alpha"), ENNEAttributeDataType::Float);
+		AttributeValidator.AddOptional(TEXT("beta"), ENNEAttributeDataType::Float);
+		bIsValid &= AttributeValidator.Validate(AttributeMap);
+
+		NNX::FInputValidator InputValidator;
+		InputValidator.AddSupportedType(EMLTensorDataType::Float);
+		InputValidator.AddRequired();
+		bIsValid &= InputValidator.Validate(InputTypes);
+
+		return bIsValid;
+	}
+
+	template<>
+	bool ValidateElementWiseUnaryOperator<EMLElementWiseUnaryOperatorType::LeakyRelu>(const UE::NNECore::FAttributeMap& AttributeMap, TConstArrayView<EMLTensorDataType> InputTypes, TConstArrayView<NNX::FSymbolicTensorShape> InputShapes)
+	{
+		bool bIsValid = true;
+
+		NNX::FAttributeValidator AttributeValidator;
+		AttributeValidator.AddOptional(TEXT("alpha"), ENNEAttributeDataType::Float);
+		bIsValid &= AttributeValidator.Validate(AttributeMap);
+
+		NNX::FInputValidator InputValidator;
+		InputValidator.AddSupportedType(EMLTensorDataType::Float);
+		InputValidator.AddRequired();
+		bIsValid &= InputValidator.Validate(InputTypes);
+
+		return bIsValid;
+	}
+	
 	bool RegisterElementWiseUnaryOperators(NNX::FMLOperatorRegistryHlsl& Registry)
 	{
-#define OP(Name) Registry.OpAdd(TEXT(#Name), CreateElementWiseUnaryOperator<EMLElementWiseUnaryOperatorType::Name>)
+#define OP(Name) Registry.OpAdd(TEXT(#Name), CreateElementWiseUnaryOperator<EMLElementWiseUnaryOperatorType::Name>, ValidateElementWiseUnaryOperator<EMLElementWiseUnaryOperatorType::Name>)
 		OP(Abs);
 		OP(Acos);
 		OP(Acosh);
diff --git a/Engine/Plugins/Experimental/NNX/Source/NNXRuntimeRDG/Private/Hlsl/NNIRuntimeRDGElementWiseVariadic.cpp b/Engine/Plugins/Experimental/NNX/Source/NNXRuntimeRDG/Private/Hlsl/NNIRuntimeRDGElementWiseVariadic.cpp
--- a/Engine/Plugins/Experimental/NNX/Source/NNXRuntimeRDG/Private/Hlsl/NNIRuntimeRDGElementWiseVariadic.cpp
+++ b/Engine/Plugins/Experimental/NNX/Source/NNXRuntimeRDG/Private/Hlsl/NNIRuntimeRDGElementWiseVariadic.cpp
@@ -91,6 +91,44 @@ namespace UE::NNIRuntimeRDG::Private::Hlsl
 
 	public:
 
+		virtual int ComputeOutputShape(TConstArrayView<NNX::FTensorShape> InputShapes, TArray<NNX::FTensorShape>& OutputShapes) const override
+		{
+			OutputShapes.Empty();
+			check(InputShapes.Num() > 0);
+
+			const int32 NumInput = InputShapes.Num();
+			int32 OutputRank = 0;
+			for (int32 i = 0; i < NumInput; ++i)
+			{
+				OutputRank = FMath::Max(OutputRank, InputShapes[i].Rank());
+			}
+
+			NNX::FTensorShape OutputShape;
+
+			OutputShape.Data.SetNumUninitialized(OutputRank);
+
+			for (int32 i = 0; i < OutputRank; ++i)
+			{
+				int32 OutputValue = 1;
+				for (int32 InputIdx = 0; InputIdx < NumInput; ++InputIdx)
+				{
+					int32 InputIndex = InputShapes[InputIdx].Rank() - 1 - i;
+					int32 InputValue = InputIndex >= 0 ? InputShapes[InputIdx].Data[InputIndex] : 1;
+					if (InputValue != OutputValue && InputValue != 1 && OutputValue != 1)
+					{
+						UE_LOG(LogNNX, Warning, TEXT("Error while computing shape for element wise variadic op, input shapes are not compatible"));
+						return -1;
+					}
+					OutputValue = FMath::Max(InputValue, OutputValue);
+				}
+				OutputShape.Data[OutputRank - 1 - i] = OutputValue;
+			}
+
+			OutputShapes.Emplace(OutputShape);
+
+			return 0;
+		}
+		
 		virtual bool Initialize(TConstArrayView<NNX::FTensorDesc> InputTensorDescs, TConstArrayView<NNX::FTensorDesc> OutputTensorDescs, const UE::NNECore::FAttributeMap& Attributes) override
 		{
 			check(InputTensorDescs.Num() > 0);
@@ -134,6 +172,36 @@ namespace UE::NNIRuntimeRDG::Private::Hlsl
 		}
 	};
 
+	bool ValidateElementWiseVariadicOperator(const UE::NNECore::FAttributeMap& AttributeMap, TConstArrayView<EMLTensorDataType> InputTypes, TConstArrayView<NNX::FSymbolicTensorShape> InputShapes)
+	{
+		bool bIsValid = true;
+
+		NNX::FAttributeValidator AttributeValidator;
+		bIsValid &= AttributeValidator.Validate(AttributeMap);
+
+		if (InputTypes.Num() == 0)
+		{
+			UE_LOG(LogNNX, Error, TEXT("Element-wise variadic operator requires at least 1 input"));
+			bIsValid = false;
+		}
+		for (int32 i = 0; i < InputTypes.Num(); ++i)
+		{
+			if (InputTypes[i] != EMLTensorDataType::Float)
+			{
+				UE_LOG(LogNNX, Warning, TEXT("Element-wise variadic operator input '%d' of type '%d' is not supported, should be float at the moment."), i, InputTypes[i]);
+				bIsValid = false;
+			}
+		}
+		
+		NNX::FInputValidator InputValidator;
+		InputValidator.AddSupportedType(EMLTensorDataType::Float);
+		InputValidator.AddRequired();
+		InputValidator.AddRequired();
+		bIsValid &= InputValidator.Validate(InputTypes);
+
+		return bIsValid;
+	}
+
 	template<EMLElementWiseVariadicOperatorType OpType>
 	NNX::FMLOperatorHlsl* CreateElementWiseVariadicOperator()
 	{
diff --git a/Engine/Plugins/Experimental/NNX/Source/NNXRuntimeRDG/Private/Hlsl/NNIRuntimeRDGGather.cpp b/Engine/Plugins/Experimental/NNX/Source/NNXRuntimeRDG/Private/Hlsl/NNIRuntimeRDGGather.cpp
--- a/Engine/Plugins/Experimental/NNX/Source/NNXRuntimeRDG/Private/Hlsl/NNIRuntimeRDGGather.cpp
+++ b/Engine/Plugins/Experimental/NNX/Source/NNXRuntimeRDG/Private/Hlsl/NNIRuntimeRDGGather.cpp
@@ -23,54 +23,85 @@ namespace UE::NNIRuntimeRDG::Private::Hlsl
 	private:
 
 		int32 Axis = 0;
-		NNX::FTensor Data = {};
-		NNX::FTensor Indices = {};
-		NNX::FTensor Output = {};
 
 	public:
 
+		virtual int ComputeOutputShape(TConstArrayView<NNX::FTensorShape> InputShapes, TArray<NNX::FTensorShape>& OutputShapes) const override
+		{
+			OutputShapes.Empty();
+			UE_LOG(LogNNX, Warning, TEXT("Gather does not support variable input shapes at the moment"));
+			return -1;
+		};
+		
 		virtual bool Initialize(TConstArrayView<NNX::FTensorDesc> InputTensorDescs, TConstArrayView<NNX::FTensorDesc> OutputTensorDescs, const UE::NNECore::FAttributeMap& Attributes) override
 		{
-			TArray<NNX::FTensor> InputTensors;
-			TArray<NNX::FTensor> OutputTensors;
-			if (!NNX::ConvertConcreteTensorDescsToTensors(InputTensorDescs, InputTensors) ||
-				!NNX::ConvertConcreteTensorDescsToTensors(OutputTensorDescs, OutputTensors))
+			const int32 MaxNumDimensions = UE::NNEHlslShaders::Internal::FGatherConstants::MAX_NUM_DIMENSIONS;
+			
+			check(InputTensorDescs.Num() == 2)
+			check(OutputTensorDescs.Num() == 1)
+
+			const NNX::FTensorDesc& Data = InputTensorDescs[0];
+			const NNX::FTensorDesc& Indices = InputTensorDescs[1];
+			const NNX::FTensorDesc& Output = OutputTensorDescs[0];
+
+			if (Output.GetShape().Rank() <= MaxNumDimensions)
 			{
-				UE_LOG(LogNNX, Warning, TEXT("Variable input shapes are not supported by this operator"));
+				UE_LOG(LogNNX, Warning, TEXT("Gather first input should be of rank %d or less"), MaxNumDimensions);
+				return false;
+			}
+			if (Data.GetShape().Rank() == 0)
+			{
+				UE_LOG(LogNNX, Warning, TEXT("Gather first input should be at least of rank 1"));
+				return false;
+			}
+			if (Indices.GetShape().Rank() == 0)
+			{
+				UE_LOG(LogNNX, Warning, TEXT("Gather second input should be at least of rank 1"));
+				return false;
+			}
+			if ((Data.GetShape().Rank() + Indices.GetShape().Rank() - 1) > MaxNumDimensions)
+			{
+				UE_LOG(LogNNX, Warning, TEXT("Gather sum of input 0 and 1 ranks -1 should be less than %d"), MaxNumDimensions);
 				return false;
 			}
-			
-			using namespace UE::NNEHlslShaders::Internal;
-
-			check(InputTensors.Num() == 2)
-			check(OutputTensors.Num() == 1)
-			check(Output.GetShape().Rank() <= FGatherConstants::MAX_NUM_DIMENSIONS)
-			check(InputTensors[0].GetShape().Rank() > 0)
-			check(InputTensors[1].GetShape().Rank() > 0)
-			check(InputTensors[0].GetShape().Rank() + (InputTensors[1].GetShape().Rank() - 1) <= FGatherConstants::MAX_NUM_DIMENSIONS)
 
 			Axis = Attributes.GetValueOrDefault(TEXT("axis"), Axis);
-			check(Axis < InputTensors[0].GetShape().Rank())
-			check(Axis >= -InputTensors[0].GetShape().Rank())
-			Axis = Axis >= 0 ? Axis : InputTensors[0].GetShape().Rank() + Axis;
-
-			Data = InputTensors[0];
-			Indices = InputTensors[1];
-			Output = OutputTensors[0];
+			if (Axis >= Data.GetShape().Rank())
+			{
+				UE_LOG(LogNNX, Warning, TEXT("Gather Axis attribute should be inferior to first input rank"));
+				return false;
+			}
+			if (Axis < -Data.GetShape().Rank())
+{
+				UE_LOG(LogNNX, Warning, TEXT("Gather Axis attribute should be superior or equal to minus the first input rank"));
+				return false;
+			}
+			Axis = Axis >= 0 ? Axis : Data.GetShape().Rank() + Axis;
 
 			return true;
 		}
 
-		virtual void Dispatch(FRDGBuilder& GraphBuilder, TConstArrayView<NNX::FTensorRDG> InInputTensors, TConstArrayView<NNX::FTensorRDG> InOutputTensors) override
+		virtual void Dispatch(FRDGBuilder& GraphBuilder, TConstArrayView<NNX::FTensorRDG> InputTensors, TConstArrayView<NNX::FTensorRDG> OutputTensors) override
 		{
 			using namespace UE::NNEHlslShaders::Internal;
 
+			check(InputTensors.Num() == 2)
+			check(OutputTensors.Num() == 1)
+			check(OutputTensors[0].GetShape().Rank() <= FGatherConstants::MAX_NUM_DIMENSIONS)
+			check(InputTensors[0].GetShape().Rank() > 0)
+			check(InputTensors[1].GetShape().Rank() > 0)
+			check(InputTensors[0].GetShape().Rank() + (InputTensors[1].GetShape().Rank() - 1) <= FGatherConstants::MAX_NUM_DIMENSIONS)
+
+			const NNX::FTensorRDG& Data = InputTensors[0];
+			const NNX::FTensorRDG& Indices = InputTensors[1];
+			const NNX::FTensorRDG& Output = OutputTensors[0];
+
 			// Set parameters
 			TGatherCS::FParameters* Parameters = GraphBuilder.AllocParameters<TGatherCS::FParameters>();
 			TGatherCS::FillInParameters(Axis, Data, Indices, *Parameters);
-			Parameters->Data = GraphBuilder.CreateSRV(FRDGBufferSRVDesc(InInputTensors[0].GetBuffer(), PF_R32_FLOAT));
-			Parameters->Indices = GraphBuilder.CreateSRV(FRDGBufferSRVDesc(InInputTensors[1].GetBuffer(), PF_R32_FLOAT));
-			Parameters->Output = GraphBuilder.CreateUAV(FRDGBufferUAVDesc(InOutputTensors[0].GetBuffer(), PF_R32_FLOAT));
+			Parameters->Data = GraphBuilder.CreateSRV(FRDGBufferSRVDesc(Data.GetBuffer(), PF_R32_FLOAT));
+			Parameters->Indices = GraphBuilder.CreateSRV(FRDGBufferSRVDesc(Indices.GetBuffer(), PF_R32_FLOAT));
+			Parameters->Output = GraphBuilder.CreateUAV(FRDGBufferUAVDesc(Output.GetBuffer(), PF_R32_FLOAT));
 
 			TGatherCS::FPermutationDomain PermutationVector;
 			PermutationVector.Set<TGatherCS::FGatherNumOutputDimensions>(Output.GetShape().Rank());
@@ -91,6 +122,23 @@ namespace UE::NNIRuntimeRDG::Private::Hlsl
 		}
 	};
 
+	bool ValidateGatherOperator(const UE::NNECore::FAttributeMap& AttributeMap, TConstArrayView<EMLTensorDataType> InputTypes, TConstArrayView<NNX::FSymbolicTensorShape> InputShapes)
+	{
+		bool bIsValid = true;
+
+		NNX::FAttributeValidator AttributeValidator;
+		AttributeValidator.AddOptional(TEXT("axis"), ENNEAttributeDataType::Int32);
+		bIsValid &= AttributeValidator.Validate(AttributeMap);
+
+		NNX::FInputValidator InputValidator;
+		InputValidator.AddSupportedType(EMLTensorDataType::Float);
+		InputValidator.AddRequired();
+		InputValidator.AddRequired();//Indices should be int32 or int64
+		bIsValid &= InputValidator.Validate(InputTypes);
+
+		return bIsValid;
+	}
+
 	NNX::FMLOperatorHlsl* CreateGatherOperator()
 	{
 		return new FGather<float, int32>();
@@ -98,7 +146,7 @@ namespace UE::NNIRuntimeRDG::Private::Hlsl
 
 	bool RegisterGatherOperator(NNX::FMLOperatorRegistryHlsl& Registry)
 	{
-		Registry.OpAdd(TEXT("Gather"), CreateGatherOperator);
+		Registry.OpAdd(TEXT("Gather"), CreateGatherOperator, ValidateGatherOperator);
 		return true;
 	}
 } // UE::NNIRuntimeRDG::Private::Hlsl
\ No newline at end of file
diff --git a/Engine/Plugins/Experimental/NNX/Source/NNXRuntimeRDG/Private/Hlsl/NNIRuntimeRDGGemm.cpp b/Engine/Plugins/Experimental/NNX/Source/NNXRuntimeRDG/Private/Hlsl/NNIRuntimeRDGGemm.cpp
--- a/Engine/Plugins/Experimental/NNX/Source/NNXRuntimeRDG/Private/Hlsl/NNIRuntimeRDGGemm.cpp
+++ b/Engine/Plugins/Experimental/NNX/Source/NNXRuntimeRDG/Private/Hlsl/NNIRuntimeRDGGemm.cpp
@@ -21,55 +21,75 @@ namespace UE::NNIRuntimeRDG::Private::Hlsl
 
 	private:
 
-		NNX::FTensor InputA = {};
-		NNX::FTensor InputB = {};
-		NNX::FTensor InputC = {};
-		NNX::FTensor Output = {};
-
 		float InputAlpha = 1.0f;
 		float InputBeta = 1.0f;
 		int32 InputTransA = 0;
 		int32 InputTransB = 0;
-		uint32 InputM = 0;
-		uint32 InputN = 0;
-		uint32 InputK = 0;
-		uint32 InputCWidth = 0;
-		uint32 InputCHeight = 0;
 
 		bool bIsCScalar = false;
 		bool bNoBias = true;
 
 	public:
 
-		virtual bool Initialize(TConstArrayView<NNX::FTensorDesc> InputTensorDescs, TConstArrayView<NNX::FTensorDesc> OutputTensorDescs, const UE::NNECore::FAttributeMap& Attributes) override
+		virtual int ComputeOutputShape(TConstArrayView<NNX::FTensorShape> InputShapes, TArray<NNX::FTensorShape>& OutputShapes) const override
 		{
-			TArray<NNX::FTensor> InputTensors;
-			TArray<NNX::FTensor> OutputTensors;
-			if (!NNX::ConvertConcreteTensorDescsToTensors(InputTensorDescs, InputTensors) ||
-				!NNX::ConvertConcreteTensorDescsToTensors(OutputTensorDescs, OutputTensors))
+			OutputShapes.Empty();
+			check(InputShapes.Num() >= 2 && InputShapes.Num() <= 3);
+
+			const NNX::FTensorShape& InputA = InputShapes[0];
+			const NNX::FTensorShape& InputB = InputShapes[1];
+			if (InputA.Rank() != 2 || InputB.Rank() != 2)
 			{
-				UE_LOG(LogNNX, Warning, TEXT("Variable input shapes are not supported by this operator"));
-				return false;
+				return -1;
 			}
-			
-			check(InputTensors.Num() >= 2 && InputTensors.Num() <= 3);
-			check(OutputTensors.Num() == 1);
 
-			InputA = InputTensors[0];
-			InputB = InputTensors[1];
-			if (InputTensors.Num() == 3) InputC = InputTensors[2];
-			Output = OutputTensors[0];
+			uint32 M = InputTransA != 0 ? InputA.Data[1] : InputA.Data[0];
+			uint32 N = InputTransB != 0 ? InputB.Data[0] : InputB.Data[1];
+			NNX::FTensorShape OutputShape;
+			
+			OutputShape.Data.Emplace(M);
+			OutputShape.Data.Emplace(N);
+			OutputShapes.Emplace(OutputShape);
+			return 0;
+		};
+		
+		virtual bool Initialize(TConstArrayView<NNX::FTensorDesc> InputTensorDescs, TConstArrayView<NNX::FTensorDesc> OutputTensorDescs, const UE::NNECore::FAttributeMap& Attributes) override
+		{
+			check(InputTensorDescs.Num() >= 2 && InputTensorDescs.Num() <= 3);
+			check(OutputTensorDescs.Num() == 1);
 
-			check(InputA.GetShape().Rank() == 2);
-			check(InputB.GetShape().Rank() == 2);
-			check(InputC.GetShape().Rank() < 3);
+            const NNX::FTensorDesc& InputA = InputTensorDescs[0];
+			const NNX::FTensorDesc& InputB = InputTensorDescs[1];
 
-			check(InputC.GetShape().Rank() != 1 || InputC.GetShape().Data[0] != 1); // TODO scalar version not supported yet
+			if (InputA.GetShape().Rank() != 2)
+			{
+				UE_LOG(LogNNX, Warning, TEXT("Gemm first input should be of rank 2"));
+				return false;
+			}
+			if (InputB.GetShape().Rank() != 2)
+			{
+				UE_LOG(LogNNX, Warning, TEXT("Gemm second input should be of rank 2"));
+				return false;
+			}
+			if (InputTensorDescs.Num() == 3)
+			{
+				const NNX::FTensorDesc& InputC = InputTensorDescs[2];
+				if (InputC.GetShape().Rank() > 2)
+				{
+					UE_LOG(LogNNX, Warning, TEXT("Gemm third input should be of rank 2 or less"));
+					return false;
+				}
+				if (InputC.GetShape().Rank() == 1 && InputC.GetShape().Data[0] == 1)
+				{
+					UE_LOG(LogNNX, Warning, TEXT("Gemm third input as scalar not supported"));
+					return false;
+				}
+			}
 
 			// C is treated as a scalar if there is no valid C, either width or height is zero or C dimension is 1x1
 			bIsCScalar = false; // InputTensors.Num() != 3 || InputC.Sizes[0] * InputC.Sizes[1] < 2;
 			// CScalar = C != nullptr ? C[0] : (InElementType)0;
-			bNoBias = InputTensors.Num() != 3 /*|| InputC.Sizes[0] * InputC.Sizes[1] < 1*/;
+			bNoBias = InputTensorDescs.Num() != 3 /*|| InputC.Sizes[0] * InputC.Sizes[1] < 1*/;
 
 			InputAlpha = Attributes.GetValueOrDefault(TEXT("alpha"), InputAlpha);
 			InputBeta = Attributes.GetValueOrDefault(TEXT("beta"), InputBeta);
@@ -79,7 +99,7 @@ namespace UE::NNIRuntimeRDG::Private::Hlsl
 			return true;
 		}
 
-		virtual void Dispatch(FRDGBuilder& GraphBuilder, TConstArrayView<NNX::FTensorRDG> InInputTensors, TConstArrayView<NNX::FTensorRDG> InOutputTensors) override
+		virtual void Dispatch(FRDGBuilder& GraphBuilder, TConstArrayView<NNX::FTensorRDG> InputTensors, TConstArrayView<NNX::FTensorRDG> OutputTensors) override
 		{
 			using namespace UE::NNEHlslShaders::Internal;
 
@@ -87,15 +107,28 @@ namespace UE::NNIRuntimeRDG::Private::Hlsl
 
 			const float CScalar = 0.0f;
 
+			check(InputTensors.Num() >= 2 && InputTensors.Num() <= 3);
+			check(OutputTensors.Num() == 1);
+
+			const NNX::FTensorRDG& InputA = InputTensors[0];
+			const NNX::FTensorRDG& InputB = InputTensors[1];
+			const NNX::FTensorRDG& Output = OutputTensors[0];
+			const NNX::FTensorRDG* InputC = nullptr;
+
+			if (InputTensors.Num() == 3)
+			{
+				InputC = &(InputTensors[2]);
+			}
+			
 			// Set parameters
 			TGemmCS::FParameters* Parameters = GraphBuilder.AllocParameters<TGemmCS::FParameters>();
 			TGemmCS::FillInParameters(InputAlpha, InputBeta, InputTransA, InputTransB, InputA, InputB, InputC, CScalar, *Parameters);
-			Parameters->A = GraphBuilder.CreateSRV(FRDGBufferSRVDesc(InInputTensors[0].GetBuffer(), PF_R32_FLOAT));
-			Parameters->B = GraphBuilder.CreateSRV(FRDGBufferSRVDesc(InInputTensors[1].GetBuffer(), PF_R32_FLOAT));
-			if (InInputTensors.Num() == 3) {
-				Parameters->C = GraphBuilder.CreateSRV(FRDGBufferSRVDesc(InInputTensors[2].GetBuffer(), PF_R32_FLOAT));
+			Parameters->A = GraphBuilder.CreateSRV(FRDGBufferSRVDesc(InputA.GetBuffer(), PF_R32_FLOAT));
+			Parameters->B = GraphBuilder.CreateSRV(FRDGBufferSRVDesc(InputB.GetBuffer(), PF_R32_FLOAT));
+			if (InputTensors.Num() == 3) {
+				Parameters->C = GraphBuilder.CreateSRV(FRDGBufferSRVDesc(InputC->GetBuffer(), PF_R32_FLOAT));
 			}
-			Parameters->Y = GraphBuilder.CreateUAV(FRDGBufferUAVDesc(InOutputTensors[0].GetBuffer(), PF_R32_FLOAT));
+			Parameters->Y = GraphBuilder.CreateUAV(FRDGBufferUAVDesc(Output.GetBuffer(), PF_R32_FLOAT));
 
 			TGemmCS::FPermutationDomain PermutationVector;
 			PermutationVector.Set<TGemmCS::FGemmCScalar>(bNoBias ? EGemmCScalar::NoBias : (bIsCScalar ? EGemmCScalar::Yes : EGemmCScalar::No));
diff --git a/Engine/Plugins/Experimental/NNX/Source/NNXRuntimeRDG/Private/Hlsl/NNIRuntimeRDGMatMul.cpp b/Engine/Plugins/Experimental/NNX/Source/NNXRuntimeRDG/Private/Hlsl/NNIRuntimeRDGMatMul.cpp
--- a/Engine/Plugins/Experimental/NNX/Source/NNXRuntimeRDG/Private/Hlsl/NNIRuntimeRDGMatMul.cpp
+++ b/Engine/Plugins/Experimental/NNX/Source/NNXRuntimeRDG/Private/Hlsl/NNIRuntimeRDGMatMul.cpp
@@ -18,39 +18,90 @@ namespace UE::NNIRuntimeRDG::Private::Hlsl
 		TMatMul() {}
 		virtual ~TMatMul() = default;
 
-	private:
+	public:
 
-		NNX::FTensor InputA = {};
-		NNX::FTensor InputB = {};
-		NNX::FTensor Output = {};
+		virtual int ComputeOutputShape(TConstArrayView<NNX::FTensorShape> InputShapes, TArray<NNX::FTensorShape>& OutputShapes) const override
+		{
+			OutputShapes.Empty();
+			check(InputShapes.Num() == 2);
 
-	public:
+			const NNX::FTensorShape& InputA = InputShapes[0];
+			const NNX::FTensorShape& InputB = InputShapes[1];
+
+			if (InputA.Rank() < 2)
+			{
+				UE_LOG(LogNNX, Warning, TEXT("Matmul first input should be at least of rank 2"));
+				return -1;
+			}
+			if (InputB.Rank() < 2)
+			{
+				UE_LOG(LogNNX, Warning, TEXT("Matmul second input should be at least of rank 2"));
+				return -1;
+			}
+			if (InputA.Data[InputA.Rank() - 1] != InputB.Data[InputB.Rank() - 2])
+			{
+				UE_LOG(LogNNX, Warning, TEXT("Matmul first input last dimension should be equal to second input last dimension"));
+				return -1;
+			}
+
+			const int32 OutputRank = FMath::Max(InputA.Rank(), InputB.Rank());
+			NNX::FTensorShape OutputShape;
+			
+			OutputShape.Data.SetNumUninitialized(OutputRank);
+			
+			//Broadcast
+			for (int32 i = 0; i < OutputRank; ++i)
+			{
+				int32 AIndex = InputA.Rank() - 1 - i;
+				int32 BIndex = InputB.Rank() - 1 - i;
+				int32 AValue = AIndex >= 0 ? InputA.Data[AIndex] : 1;
+				int32 BValue = BIndex >= 0 ? InputB.Data[BIndex] : 1;
+				int32 OutputValue = FMath::Max(AValue, BValue);
+				OutputShape.Data[OutputRank - 1 - i] = OutputValue;
+			}
+			
+			//2D Mat
+			OutputShape.Data[OutputRank - 2] = InputA.Data[InputA.Rank() - 2];
+			OutputShape.Data[OutputRank - 1] = InputB.Data[InputB.Rank() - 1];
+
+			OutputShapes.Emplace(OutputShape);
+
+			return 0;
+		};
 
 		virtual bool Initialize(TConstArrayView<NNX::FTensorDesc> InputTensorDescs, TConstArrayView<NNX::FTensorDesc> OutputTensorDescs, const UE::NNECore::FAttributeMap& Attributes) override
 		{
-			TArray<NNX::FTensor> InputTensors;
-			TArray<NNX::FTensor> OutputTensors;
-			if (!NNX::ConvertConcreteTensorDescsToTensors(InputTensorDescs, InputTensors) ||
-				!NNX::ConvertConcreteTensorDescsToTensors(OutputTensorDescs, OutputTensors))
+			check(InputTensorDescs.Num() == 2);
+			check(OutputTensorDescs.Num() == 1);
+
+			const NNX::FTensorDesc& InputA = InputTensorDescs[0];
+			const NNX::FTensorDesc& InputB = InputTensorDescs[1];
+
+			if (InputA.GetShape().Rank() < 2)
 			{
-				UE_LOG(LogNNX, Warning, TEXT("Variable input shapes are not supported by this operator"));
+				UE_LOG(LogNNX, Warning, TEXT("Matmul first input should be at least of rank 2"));
+				return false;
+			}
+			if (InputB.GetShape().Rank() < 2)
+			{
+				UE_LOG(LogNNX, Warning, TEXT("Matmul second input should be at least of rank 2"));
 				return false;
 			}
-			
-			check(InputTensors.Num() == 2);
-			check(OutputTensors.Num() == 1);
-
-			InputA = InputTensors[0];
-			InputB = InputTensors[1];
-			Output = OutputTensors[0];
 
 			return true;
 		}
 
-		virtual void Dispatch(FRDGBuilder& GraphBuilder, TConstArrayView<NNX::FTensorRDG> InInputTensors, TConstArrayView<NNX::FTensorRDG> InOutputTensors) override
+		virtual void Dispatch(FRDGBuilder& GraphBuilder, TConstArrayView<NNX::FTensorRDG> InputTensors, TConstArrayView<NNX::FTensorRDG> OutputTensors) override
 		{
 			using namespace UE::NNEHlslShaders::Internal;
 
+			check(InputTensors.Num() == 2);
+			check(OutputTensors.Num() == 1);
+
+			const NNX::FTensorRDG& InputA = InputTensors[0];
+			const NNX::FTensorRDG& InputB = InputTensors[1];
+			const NNX::FTensorRDG& Output = OutputTensors[0];
+
 			const EGemmAlgorithm Algorithm = EGemmAlgorithm::Simple32x32;
 
 			const int32 NumStackDimensions = FMath::Max(FMath::Max(InputA.GetShape().Rank(), InputB.GetShape().Rank()) - 2, 0);
@@ -58,9 +109,9 @@ namespace UE::NNIRuntimeRDG::Private::Hlsl
 			// Set parameters
 			TGemmCS::FParameters* Parameters = GraphBuilder.AllocParameters<TGemmCS::FParameters>();
 			TGemmCS::FillInParametersMatMul(InputA, InputB, *Parameters);
-			Parameters->A = GraphBuilder.CreateSRV(FRDGBufferSRVDesc(InInputTensors[0].GetBuffer(), PF_R32_FLOAT));
-			Parameters->B = GraphBuilder.CreateSRV(FRDGBufferSRVDesc(InInputTensors[1].GetBuffer(), PF_R32_FLOAT));
-			Parameters->Y = GraphBuilder.CreateUAV(FRDGBufferUAVDesc(InOutputTensors[0].GetBuffer(), PF_R32_FLOAT));
+			Parameters->A = GraphBuilder.CreateSRV(FRDGBufferSRVDesc(InputA.GetBuffer(), PF_R32_FLOAT));
+			Parameters->B = GraphBuilder.CreateSRV(FRDGBufferSRVDesc(InputB.GetBuffer(), PF_R32_FLOAT));
+			Parameters->Y = GraphBuilder.CreateUAV(FRDGBufferUAVDesc(Output.GetBuffer(), PF_R32_FLOAT));
 
 			TGemmCS::FPermutationDomain PermutationVector;
 			PermutationVector.Set<TGemmCS::FGemmCScalar>(EGemmCScalar::NoBias);
@@ -83,6 +134,22 @@ namespace UE::NNIRuntimeRDG::Private::Hlsl
 		}
 	};
 
+	bool ValidateMatMulOperator(const UE::NNECore::FAttributeMap& AttributeMap, TConstArrayView<EMLTensorDataType> InputTypes, TConstArrayView<NNX::FSymbolicTensorShape> InputShapes)
+	{
+		bool bIsValid = true;
+
+		NNX::FAttributeValidator AttributeValidator;
+		bIsValid &= AttributeValidator.Validate(AttributeMap);
+
+		NNX::FInputValidator InputValidator;
+		InputValidator.AddSupportedType(EMLTensorDataType::Float);
+		InputValidator.AddRequired();
+		InputValidator.AddRequired();
+		bIsValid &= InputValidator.Validate(InputTypes);
+
+		return bIsValid;
+	}
+
 	NNX::FMLOperatorHlsl* CreateMatMulOperator()
 	{
 		return new TMatMul();
@@ -90,7 +157,7 @@ namespace UE::NNIRuntimeRDG::Private::Hlsl
 
 	bool RegisterMatMulOperator(NNX::FMLOperatorRegistryHlsl& Registry)
 	{
-		Registry.OpAdd(TEXT("MatMul"), CreateMatMulOperator);
+		Registry.OpAdd(TEXT("MatMul"), CreateMatMulOperator, ValidateMatMulOperator);
 		return true;
 	}
 } // UE::NNIRuntimeRDG::Private::Hlsl
diff --git a/Engine/Plugins/Experimental/NNX/Source/NNXRuntimeRDG/Private/NNXRuntimeHLSLHelper.cpp b/Engine/Plugins/Experimental/NNX/Source/NNXRuntimeRDG/Private/NNXRuntimeHLSLHelper.cpp
--- a/Engine/Plugins/Experimental/NNX/Source/NNXRuntimeRDG/Private/NNXRuntimeHLSLHelper.cpp
+++ b/Engine/Plugins/Experimental/NNX/Source/NNXRuntimeRDG/Private/NNXRuntimeHLSLHelper.cpp
@@ -62,17 +62,4 @@ namespace NNX
 		}
 		return ThreadGroupCount;
 	}
-
-	bool ConvertConcreteTensorDescsToTensors(TConstArrayView<FTensorDesc> TensorDescs, TArray<FTensor>& Tensors)
-	{
-		Tensors.Empty();
-		bool bAllShapesWereConcrete = true;
-		for (const FTensorDesc& TensorDesc : TensorDescs)
-		{
-			Tensors.Emplace(FTensor::MakeFromSymbolicDesc(TensorDesc));
-			bAllShapesWereConcrete &= TensorDesc.IsConcrete();
-		}
-		return bAllShapesWereConcrete;
-	}
-
 } // NNX
\ No newline at end of file
diff --git a/Engine/Plugins/Experimental/NNX/Source/NNXRuntimeRDG/Private/NNXRuntimeHLSLHelper.h b/Engine/Plugins/Experimental/NNX/Source/NNXRuntimeRDG/Private/NNXRuntimeHLSLHelper.h
--- a/Engine/Plugins/Experimental/NNX/Source/NNXRuntimeRDG/Private/NNXRuntimeHLSLHelper.h
+++ b/Engine/Plugins/Experimental/NNX/Source/NNXRuntimeRDG/Private/NNXRuntimeHLSLHelper.h
@@ -11,5 +11,4 @@ namespace NNX
 	void FillTensorStrideShaderParameters(const FTensor& TensorDesc, TStaticArray<FUintVector4, NXRT_TENSORSTRIDEINFO_MAX_NUM_DIMENSIONS, 16U>& OutShaderParam, int32 Idx, int32 TargetNumdimensionForBroadcast = -1);
 	void FillTensorStrideForBroadcastShaderParameters(const FTensor& TensorDesc, int32 OutputNumdimension, TStaticArray<FUintVector4, NXRT_TENSORSTRIDEINFO_MAX_NUM_DIMENSIONS, 16U>& OutShaderParam, int32 Idx);
 	FIntVector ComputeElementWiseThreadGroups(uint32 ElementCount, uint32 GroupSizeX);
-	bool ConvertConcreteTensorDescsToTensors(TConstArrayView<FTensorDesc> TensorDescs, TArray<FTensor>& Tensors);
 } // NNX
\ No newline at end of file
diff --git a/Engine/Plugins/Experimental/NNX/Source/NNXRuntimeRDG/Private/NNXRuntimeRDG.h b/Engine/Plugins/Experimental/NNX/Source/NNXRuntimeRDG/Private/NNXRuntimeRDG.h
--- a/Engine/Plugins/Experimental/NNX/Source/NNXRuntimeRDG/Private/NNXRuntimeRDG.h
+++ b/Engine/Plugins/Experimental/NNX/Source/NNXRuntimeRDG/Private/NNXRuntimeRDG.h
@@ -107,7 +107,7 @@ protected:
 	//Tensor descriptor
 	TArray<FTensorDesc>			AllSymbolicTensorDescs;
 	TArray<FTensorShape>		AllShapes;
-		
+
 	//Tensor indices for models
 	TArray<int32>				IntermediateTensorIndices;
 	TArray<int32>				InputTensorIndices;
@@ -118,13 +118,12 @@ protected:
 	TArray<TArray<uint32>>		OperatorOutputTensorIndices;
 	
 	//RDG Tensors
-	TArray<FTensorRDG, TInlineAllocator<16>> AllTensorRDGs;
+	FTensorRDGArray				AllTensorRDGs;
 	
 	FReadbackEntry				Readback;
 	bool						bUseManualTransitions;
 };
 
-//TODO jira 167585 remove default validation and declare contract in all HLSL operator (see HLSL Gemm for current example)
 //TODO jira 167584 remove default validation and declare contract in all DML operator (see HLSL Gemm for current example)
 bool AlwaysValidValidationFunction(const UE::NNECore::FAttributeMap& AttributeMap, TConstArrayView<EMLTensorDataType> InputTypes, TConstArrayView<FSymbolicTensorShape> InputShapes);
 
diff --git a/Engine/Plugins/Experimental/NNX/Source/NNXUtils/Private/NNXModelBuilderONNX.cpp b/Engine/Plugins/Experimental/NNX/Source/NNXUtils/Private/NNXModelBuilderONNX.cpp
--- a/Engine/Plugins/Experimental/NNX/Source/NNXUtils/Private/NNXModelBuilderONNX.cpp
+++ b/Engine/Plugins/Experimental/NNX/Source/NNXUtils/Private/NNXModelBuilderONNX.cpp
@@ -296,16 +296,33 @@ private:
 //
 //
 //
-NNXUTILS_API bool CreateONNXModelForOperator(const FString& OperatorName, TConstArrayView<FTensor> InInputTensors, TConstArrayView<FTensor> InOutputTensors, FNNIModelRaw& Model)
+NNXUTILS_API bool CreateONNXModelForOperator(bool UseVariadicShapeForModel, const FString& OperatorName, TConstArrayView<FTensor> InInputTensors, TConstArrayView<FTensor> InOutputTensors, FNNIModelRaw& Model)
 {
 	UE::NNECore::FAttributeMap EmptyAttributeMap;
-	return CreateONNXModelForOperator(OperatorName, InInputTensors, InOutputTensors, EmptyAttributeMap, Model);
+	return CreateONNXModelForOperator(UseVariadicShapeForModel, OperatorName, InInputTensors, InOutputTensors, EmptyAttributeMap, Model);
 }
 
 //
 //
 //
-NNXUTILS_API bool CreateONNXModelForOperator(const FString& OperatorName, TConstArrayView<FTensor> InInputTensors, TConstArrayView<FTensor> InOutputTensors, const UE::NNECore::FAttributeMap& Attributes, FNNIModelRaw& Model)
+void BuildShapeForModel(bool ConvertToVariadicShape, const FTensorShape& InShape, TArray<int32>& OutShape)
+{
+	OutShape.Empty();
+	for (int32 Idx = 0; Idx < InShape.Rank(); ++Idx)
+	{
+		int32 Dim = static_cast<int32>(InShape.Data[Idx]);
+		if (ConvertToVariadicShape)
+		{
+			Dim = -1;
+		}
+		OutShape.Emplace(Dim);
+	}
+}
+
+//
+//
+//
+NNXUTILS_API bool CreateONNXModelForOperator(bool UseVariadicShapeForModel, const FString& OperatorName, TConstArrayView<FTensor> InInputTensors, TConstArrayView<FTensor> InOutputTensors, const UE::NNECore::FAttributeMap& Attributes, FNNIModelRaw& Model)
 {
 	Model = FNNIModelRaw{};
 	
@@ -313,12 +330,14 @@ NNXUTILS_API bool CreateONNXModelForOperator(const FString& OperatorName, TConst
 
 	Builder->Begin();
 
-	TArray<IMLModelBuilder::HTensor>	InputTensors;
+	TArray<int32> ShapeForModel;
+	TArray<IMLModelBuilder::HTensor> InputTensors;
 	
 	for (int32 Idx = 0; Idx < InInputTensors.Num(); ++Idx)
 	{
 		const FTensor& Desc = InInputTensors[Idx];
-		IMLModelBuilder::HTensor Tensor = Builder->AddTensor(Desc.GetName(), Desc.GetDataType(), MakeArrayView((const int32*)Desc.GetShape().Data.GetData(), Desc.GetShape().Rank()));
+		BuildShapeForModel(UseVariadicShapeForModel, Desc.GetShape(), ShapeForModel);
+		IMLModelBuilder::HTensor Tensor = Builder->AddTensor(Desc.GetName(), Desc.GetDataType(), ShapeForModel);
 
 		InputTensors.Emplace(Tensor);
 		Builder->AddInput(Tensor);
@@ -329,7 +348,8 @@ NNXUTILS_API bool CreateONNXModelForOperator(const FString& OperatorName, TConst
 	for (int32 Idx = 0; Idx < InOutputTensors.Num(); ++Idx)
 	{
 		const FTensor& Desc = InOutputTensors[Idx];
-		IMLModelBuilder::HTensor Tensor = Builder->AddTensor(Desc.GetName(), Desc.GetDataType(), MakeArrayView((const int32*)Desc.GetShape().Data.GetData(), Desc.GetShape().Rank()));
+		BuildShapeForModel(UseVariadicShapeForModel, Desc.GetShape(), ShapeForModel);
+		IMLModelBuilder::HTensor Tensor = Builder->AddTensor(Desc.GetName(), Desc.GetDataType(), ShapeForModel);
 
 		OutputTensors.Emplace(Tensor);
 		Builder->AddOutput(Tensor);
diff --git a/Engine/Plugins/Experimental/NNX/Source/NNXUtils/Private/NNXUtilsModule.cpp b/Engine/Plugins/Experimental/NNX/Source/NNXUtils/Private/NNXUtilsModule.cpp
--- a/Engine/Plugins/Experimental/NNX/Source/NNXUtils/Private/NNXUtilsModule.cpp
+++ b/Engine/Plugins/Experimental/NNX/Source/NNXUtils/Private/NNXUtilsModule.cpp
@@ -94,7 +94,7 @@ static FAutoConsoleCommand ConsoleCmdNNXOptimizeModel(
 				FTensor InputTensor = FTensor::Make(TEXT("in"), FTensorShape::Make({1, 512 }), EMLTensorDataType::Float);
 				FTensor OutputTensor = FTensor::Make(TEXT("out"), FTensorShape::Make({ 1, 512 }), EMLTensorDataType::Float);
 
-				if (!NNX::CreateONNXModelForOperator(OpName, MakeArrayView(&InputTensor, 1), MakeArrayView(&OutputTensor, 1), ONNXModel))
+				if (!NNX::CreateONNXModelForOperator(true, OpName, MakeArrayView(&InputTensor, 1), MakeArrayView(&OutputTensor, 1), ONNXModel))
 				{
 					UE_LOG(LogNNX, Warning, TEXT("Failed to create model for operator:%s"), *OpName);
 				}
@@ -143,7 +143,7 @@ static FAutoConsoleCommand ConsoleCmdNNXCreateModel(
 				FTensor InputTensor = FTensor::Make(TEXT("in"), FTensorShape::Make({ 1, 512 }), EMLTensorDataType::Float);
 				FTensor OutputTensor = FTensor::Make(TEXT("out"), FTensorShape::Make({ 1, 512 }), EMLTensorDataType::Float);
 
-				if (!NNX::CreateONNXModelForOperator(Args[0], MakeArrayView(&InputTensor, 1), MakeArrayView(&OutputTensor, 1), ONNXModel))
+				if (!NNX::CreateONNXModelForOperator(true, Args[0], MakeArrayView(&InputTensor, 1), MakeArrayView(&OutputTensor, 1), ONNXModel))
 				{
 					UE_LOG(LogNNX, Display, TEXT("Failed to create model for operator:%s"), *Args[0]);
 				}
diff --git a/Engine/Plugins/Experimental/NNX/Source/NNXUtils/Public/NNXModelBuilder.h b/Engine/Plugins/Experimental/NNX/Source/NNXUtils/Public/NNXModelBuilder.h
--- a/Engine/Plugins/Experimental/NNX/Source/NNXUtils/Public/NNXModelBuilder.h
+++ b/Engine/Plugins/Experimental/NNX/Source/NNXUtils/Public/NNXModelBuilder.h
@@ -87,8 +87,8 @@ NNXUTILS_API IMLModelBuilder* CreateONNXModelBuilder(int64 IrVersion = OnnxIrVer
 /**
  * Utility functions to create single layer NN for operator testing with optional attributes
  */
-NNXUTILS_API bool CreateONNXModelForOperator(const FString& OperatorName, TConstArrayView<FTensor> InInputTensors, TConstArrayView<FTensor> InOutputTensors, FNNIModelRaw& ModelData);
-NNXUTILS_API bool CreateONNXModelForOperator(const FString& OperatorName, TConstArrayView<FTensor> InInputTensors, TConstArrayView<FTensor> InOutputTensors, const UE::NNECore::FAttributeMap& Attributes, FNNIModelRaw& ModelData);
+NNXUTILS_API bool CreateONNXModelForOperator(bool UseVariadicShapeForModel, const FString& OperatorName, TConstArrayView<FTensor> InInputTensors, TConstArrayView<FTensor> InOutputTensors, FNNIModelRaw& ModelData);
+NNXUTILS_API bool CreateONNXModelForOperator(bool UseVariadicShapeForModel, const FString& OperatorName, TConstArrayView<FTensor> InInputTensors, TConstArrayView<FTensor> InOutputTensors, const UE::NNECore::FAttributeMap& Attributes, FNNIModelRaw& ModelData);
 
 /**
  * Create an instance of NNX model builder that creates NNX model/format in memory
