Added Initializer tensor format type for tensors that contain data.

#rnx

diff --git a/Engine/Plugins/Experimental/NNX/Source/NNXCore/Public/NNXRuntimeFormat.h b/Engine/Plugins/Experimental/NNX/Source/NNXCore/Public/NNXRuntimeFormat.h
--- a/Engine/Plugins/Experimental/NNX/Source/NNXCore/Public/NNXRuntimeFormat.h
+++ b/Engine/Plugins/Experimental/NNX/Source/NNXCore/Public/NNXRuntimeFormat.h
@@ -14,7 +14,8 @@ enum class EMLFormatTensorType : uint8
 	None,	
 	Input,
 	Output,
-	Intermediate
+	Intermediate,
+	Initializer
 };
 
 //TODO jira 167589: convert EMLInferenceFormat to a CCCC for easier extension of the framework
diff --git a/Engine/Plugins/Experimental/NNX/Source/NNXUtils/Private/NNXModelBuilderNNX.cpp b/Engine/Plugins/Experimental/NNX/Source/NNXUtils/Private/NNXModelBuilderNNX.cpp
--- a/Engine/Plugins/Experimental/NNX/Source/NNXUtils/Private/NNXModelBuilderNNX.cpp
+++ b/Engine/Plugins/Experimental/NNX/Source/NNXUtils/Private/NNXModelBuilderNNX.cpp
@@ -260,6 +260,7 @@ private:
 			
 			if (Data)
 			{
+				Desc.Type = EMLFormatTensorType::Initializer;
 				Desc.DataOffset = Format.TensorData.AddUninitialized(DataSize);
 				Desc.DataSize = DataSize;
 
