NNI Dynamic shape test for ORTDml

Problem:
* All dyn shape test were deactivated for NNXRuntimeORTDml

Solution:
* Identify that Gemm test with 2 dyn inputs are failing with ORTDml losing the GPU device, and disable only those tests.

Tests:
* All NNX tests pass under win64 editor

diff --git a/Engine/Plugins/Experimental/NNX/Source/NNXQA/Private/NNXQAUtils.cpp b/Engine/Plugins/Experimental/NNX/Source/NNXQA/Private/NNXQAUtils.cpp
--- a/Engine/Plugins/Experimental/NNX/Source/NNXQA/Private/NNXQAUtils.cpp
+++ b/Engine/Plugins/Experimental/NNX/Source/NNXQA/Private/NNXQAUtils.cpp
@@ -515,10 +515,9 @@ namespace Test
 			{
 				bool bShouldRunVariadicTest = (ONNXModelVariadic.Format != ENNXInferenceFormat::Invalid);
 				bShouldRunVariadicTest &= !(RuntimeName == "NNXRuntimeDML");
-				bShouldRunVariadicTest &= !(RuntimeName == "NNXRuntimeORTDml");
 
 				bool bTestSuceeded = RunTestInferenceAndCompareToRef(TestSetup, Runtime, ONNXModel, RefOutputMemBuffers, RefOutputTensorDescs);
-				
+
 				if (bShouldRunVariadicTest)
 				{
 					if (!bTestSuceeded)
diff --git a/Engine/Plugins/Experimental/NNX/Source/NNXQA/Resources/NNXQATestDesc.json b/Engine/Plugins/Experimental/NNX/Source/NNXQA/Resources/NNXQATestDesc.json
--- a/Engine/Plugins/Experimental/NNX/Source/NNXQA/Resources/NNXQATestDesc.json
+++ b/Engine/Plugins/Experimental/NNX/Source/NNXQA/Resources/NNXQATestDesc.json
@@ -428,42 +428,69 @@
 			]
 		},
 		{
-			"name": "Gemm.TwoInputs",
+			"name": "Gemm.TwoInputs.NoTranspose",
 			"datasets": [
 				{ "inputs": [ { "shape": [ 2, 2 ] }, { "shape": [ 2, 2 ] } ], "outputs": [ { "shape": [ 2, 2 ] } ]},
 				{ "inputs": [ { "shape": [ 3, 4 ] }, { "shape": [ 4, 5 ] } ], "outputs": [ { "shape": [ 3, 5 ] } ]}
+			],
+			"runtimes": [
+				{ "name": "NNXRuntimeORTDml", "skip": true, "comment": "Variable dimension without bias is failing for NNXRuntimeORTDml" }
 			]
 		},
 		{
-			"name": "Gemm.ThreeInputs",
+			"name": "Gemm.TwoInputs.TransposedA",
+			"datasets": [
+				{ "inputs": [ { "shape": [ 2, 2 ] }, { "shape": [ 2, 2 ] } ], "outputs": [ { "shape": [ 2, 2 ] } ]},
+				{ "inputs": [ { "shape": [ 4, 3 ] }, { "shape": [ 4, 5 ] } ], "outputs": [ { "shape": [ 3, 5 ] } ]}
+			],
+			"runtimes": [
+				{ "name": "NNXRuntimeORTDml", "skip": true, "comment": "Variable dimension without bias is failing for NNXRuntimeORTDml" }
+			]
+		},
+		{
+			"name": "Gemm.TwoInputs.TransposedB",
+			"datasets": [
+				{ "inputs": [ { "shape": [ 2, 2 ] }, { "shape": [ 2, 2 ] } ], "outputs": [ { "shape": [ 2, 2 ] } ]},
+				{ "inputs": [ { "shape": [ 3, 4 ] }, { "shape": [ 5, 4 ] } ], "outputs": [ { "shape": [ 3, 5 ] } ]}
+			],
+			"runtimes": [
+				{ "name": "NNXRuntimeORTDml", "skip": true, "comment": "Variable dimension without bias is failing for NNXRuntimeORTDml" }
+			]
+		},
+		{
+			"name": "Gemm.TwoInputs.AllTransposed",
+			"datasets": [
+				{ "inputs": [ { "shape": [ 2, 2 ] }, { "shape": [ 2, 2 ] } ], "outputs": [ { "shape": [ 2, 2 ] } ]},
+				{ "inputs": [ { "shape": [ 4, 3 ] }, { "shape": [ 5, 4 ] } ], "outputs": [ { "shape": [ 3, 5 ] } ]}
+			],
+			"runtimes": [
+				{ "name": "NNXRuntimeORTDml", "skip": true, "comment": "Variable dimension without bias is failing for NNXRuntimeORTDml" }
+			]
+		},
+		{
+			"name": "Gemm.ThreeInputs.NoTranspose",
 			"datasets": [
 				{ "inputs": [ { "shape": [ 3, 4 ] }, { "shape": [ 4, 5 ] }, { "shape": [ 3, 5 ] } ], "outputs": [ { "shape": [ 3, 5 ] } ]},
 				{ "inputs": [ { "shape": [ 1, 1 ] }, { "shape": [ 1, 5 ] }, { "shape": [ 5 ] } ], "outputs": [ { "shape": [ 1, 5 ] } ]}
 			]
 		},
 		{
-			"name": "Gemm.TransposedA",
+			"name": "Gemm.ThreeInputs.TransposedA",
 			"datasets": [
-				{ "inputs": [ { "shape": [ 2, 2 ] }, { "shape": [ 2, 2 ] } ], "outputs": [ { "shape": [ 2, 2 ] } ]},
-				{ "inputs": [ { "shape": [ 4, 3 ] }, { "shape": [ 4, 5 ] } ], "outputs": [ { "shape": [ 3, 5 ] } ]},
 				{ "inputs": [ { "shape": [ 4, 3 ] }, { "shape": [ 4, 5 ] }, { "shape": [ 3, 5 ] } ], "outputs": [ { "shape": [ 3, 5 ] } ]},
 				{ "inputs": [ { "shape": [ 1, 1 ] }, { "shape": [ 1, 5 ] }, { "shape": [ 5 ] } ], "outputs": [ { "shape": [ 1, 5 ] } ]}
 			]
 		},
 		{
-			"name": "Gemm.TransposedB",
+			"name": "Gemm.ThreeInputs.TransposedB",
 			"datasets": [
-				{ "inputs": [ { "shape": [ 2, 2 ] }, { "shape": [ 2, 2 ] } ], "outputs": [ { "shape": [ 2, 2 ] } ]},
-				{ "inputs": [ { "shape": [ 3, 4 ] }, { "shape": [ 5, 4 ] } ], "outputs": [ { "shape": [ 3, 5 ] } ]},
 				{ "inputs": [ { "shape": [ 3, 4 ] }, { "shape": [ 5, 4 ] }, { "shape": [ 3, 5 ] } ], "outputs": [ { "shape": [ 3, 5 ] } ]},
 				{ "inputs": [ { "shape": [ 1, 1 ] }, { "shape": [ 5, 1 ] }, { "shape": [ 5 ] } ], "outputs": [ { "shape": [ 1, 5 ] } ]}
 			]
 		},
 		{
-			"name": "Gemm.AllTransposed",
+			"name": "Gemm.ThreeInputs.AllTransposed",
 			"datasets": [
-				{ "inputs": [ { "shape": [ 2, 2 ] }, { "shape": [ 2, 2 ] } ], "outputs": [ { "shape": [ 2, 2 ] } ]},
-				{ "inputs": [ { "shape": [ 4, 3 ] }, { "shape": [ 5, 4 ] } ], "outputs": [ { "shape": [ 3, 5 ] } ]},
 				{ "inputs": [ { "shape": [ 4, 3 ] }, { "shape": [ 5, 4 ] }, { "shape": [ 3, 5 ] } ], "outputs": [ { "shape": [ 3, 5 ] } ]},
 				{ "inputs": [ { "shape": [ 1, 1 ] }, { "shape": [ 5, 1 ] }, { "shape": [ 5 ] } ], "outputs": [ { "shape": [ 1, 5 ] } ]}
 			]
@@ -539,7 +566,7 @@
 			]
 		},
 		{
-			"name": "Gemm.TwoInputs.Alpha",
+			"name": "Gemm.TwoInputs.NoTranspose.Alpha",
 			"attribute_maps": [
 				{ "attributes": [ { "name": "alpha", "value": 0.0 } ] },
 				{ "attributes": [ { "name": "alpha", "value": 0.5 } ] },
@@ -548,8 +575,8 @@
 			]
 		},
 		{
-			"name": "Gemm.ThreeInputs.Beta",
-			"multiply_with_attribute_sets": [ "Gemm.TwoInputs.Alpha" ],
+			"name": "Gemm.ThreeInputs.NoTranspose.Beta",
+			"multiply_with_attribute_sets": [ "Gemm.TwoInputs.NoTranspose.Alpha" ],
 			"attribute_maps": [
 				{ "attributes": [ { "name": "beta", "value": 0.0 } ] },
 				{ "attributes": [ { "name": "beta", "value": 0.5 } ] },
@@ -558,31 +585,49 @@
 			]
 		},
 		{
-			"name": "Gemm.TwoInputs.DefaultTranspose",
+			"name": "Gemm.TwoInputs.NoTranspose",
 			"attribute_maps": [
 				{ "attributes": [ { "name": "transA", "value": 0, "type": "Int32" }, { "name": "transB", "value": 0, "type": "Int32" } ] }
 			]
 		},
 		{
-			"name": "Gemm.ThreeInputs.DefaultTranspose",
+			"name": "Gemm.ThreeInputs.NoTranspose",
 			"attribute_maps": [
 				{ "attributes": [ { "name": "transA", "value": 0, "type": "Int32" }, { "name": "transB", "value": 0, "type": "Int32" } ] }
 			]
 		},
 		{
-			"name": "Gemm.TransposedA",
+			"name": "Gemm.TwoInputs.TransposedA",
+			"attribute_maps": [
+				{ "attributes": [ { "name": "transA", "value": 1, "type": "Int32" } ] }
+			]
+		},
+		{
+			"name": "Gemm.TwoInputs.TransposedB",
+			"attribute_maps": [
+				{ "attributes": [ { "name": "transB", "value": 1, "type": "Int32" } ] }
+			]
+		},
+		{
+			"name": "Gemm.TwoInputs.AllTransposed",
+			"attribute_maps": [
+				{ "attributes": [ { "name": "transA", "value": 1, "type": "Int32" }, { "name": "transB", "value": 1, "type": "Int32" } ] }
+			]
+		},
+				{
+			"name": "Gemm.ThreeInputs.TransposedA",
 			"attribute_maps": [
 				{ "attributes": [ { "name": "transA", "value": 1, "type": "Int32" } ] }
 			]
 		},
 		{
-			"name": "Gemm.TransposedB",
+			"name": "Gemm.ThreeInputs.TransposedB",
 			"attribute_maps": [
 				{ "attributes": [ { "name": "transB", "value": 1, "type": "Int32" } ] }
 			]
 		},
 		{
-			"name": "Gemm.AllTransposed",
+			"name": "Gemm.ThreeInputs.AllTransposed",
 			"attribute_maps": [
 				{ "attributes": [ { "name": "transA", "value": 1, "type": "Int32" }, { "name": "transB", "value": 1, "type": "Int32" } ] }
 			]
