NNI Update to make static analysis tool not to complain about derefenced nullptr in NNIRuntimeRDGGemm.cpp

diff --git a/Engine/Plugins/Experimental/NNX/Source/NNEHlslShaders/Private/NNEHlslShadersGemmCS.cpp b/Engine/Plugins/Experimental/NNX/Source/NNEHlslShaders/Private/NNEHlslShadersGemmCS.cpp
--- a/Engine/Plugins/Experimental/NNX/Source/NNEHlslShaders/Private/NNEHlslShadersGemmCS.cpp
+++ b/Engine/Plugins/Experimental/NNX/Source/NNEHlslShaders/Private/NNEHlslShadersGemmCS.cpp
@@ -102,7 +102,7 @@ namespace UE::NNEHlslShaders::Internal
 		Parameters.MxN = M * N;
 		Parameters.CWidth = 0;
 		Parameters.CHeight = 0;
-		if (InputC)
+		if (InputC != nullptr)
 		{
 			Parameters.CWidth = InputC->GetShape().Rank() == 0 ? 0 : InputC->GetShape().Data[InputC->GetShape().Rank() - 1];;
 			Parameters.CHeight = InputC->GetShape().Rank() == 0 ? 0 : InputC->GetShape().Rank() == 1 ? 1 : InputC->GetShape().Data[InputC->GetShape().Rank() - 2];
diff --git a/Engine/Plugins/Experimental/NNX/Source/NNXRuntimeRDG/Private/Hlsl/NNIRuntimeRDGGemm.cpp b/Engine/Plugins/Experimental/NNX/Source/NNXRuntimeRDG/Private/Hlsl/NNIRuntimeRDGGemm.cpp
--- a/Engine/Plugins/Experimental/NNX/Source/NNXRuntimeRDG/Private/Hlsl/NNIRuntimeRDGGemm.cpp
+++ b/Engine/Plugins/Experimental/NNX/Source/NNXRuntimeRDG/Private/Hlsl/NNIRuntimeRDGGemm.cpp
@@ -125,7 +125,7 @@ namespace UE::NNIRuntimeRDG::Private::Hlsl
 			TGemmCS::FillInParameters(InputAlpha, InputBeta, InputTransA, InputTransB, InputA, InputB, InputC, CScalar, *Parameters);
 			Parameters->A = GraphBuilder.CreateSRV(FRDGBufferSRVDesc(InputA.GetBuffer(), PF_R32_FLOAT));
 			Parameters->B = GraphBuilder.CreateSRV(FRDGBufferSRVDesc(InputB.GetBuffer(), PF_R32_FLOAT));
-			if (InputTensors.Num() == 3) {
+			if (InputC != nullptr) {
 				Parameters->C = GraphBuilder.CreateSRV(FRDGBufferSRVDesc(InputC->GetBuffer(), PF_R32_FLOAT));
 			}
 			Parameters->Y = GraphBuilder.CreateUAV(FRDGBufferUAVDesc(Output.GetBuffer(), PF_R32_FLOAT));
