Fixed GPU buffer readback for multiple outputs.

diff --git a/Engine/Plugins/Experimental/NNX/Source/NNXQA/Resources/NNXQATestDesc.json b/Engine/Plugins/Experimental/NNX/Source/NNXQA/Resources/NNXQATestDesc.json
--- a/Engine/Plugins/Experimental/NNX/Source/NNXQA/Resources/NNXQATestDesc.json
+++ b/Engine/Plugins/Experimental/NNX/Source/NNXQA/Resources/NNXQATestDesc.json
@@ -77,6 +77,10 @@
 				{
 					"target": "TwoLayersNoWeight",
 					"tags": [ "SmokeTest" ]
+				},
+				{
+					"target": "MultiOutput",
+					"tags": [ "SmokeTest" ]
 				}
 			]
 		},
diff --git a/Engine/Plugins/Experimental/NNX/Source/NNXRuntimeRDG/Private/NNXRuntimeRDG.cpp b/Engine/Plugins/Experimental/NNX/Source/NNXRuntimeRDG/Private/NNXRuntimeRDG.cpp
--- a/Engine/Plugins/Experimental/NNX/Source/NNXRuntimeRDG/Private/NNXRuntimeRDG.cpp
+++ b/Engine/Plugins/Experimental/NNX/Source/NNXRuntimeRDG/Private/NNXRuntimeRDG.cpp
@@ -6,12 +6,21 @@
 
 #include "RenderGraphBuilder.h"
 #include "RenderGraphUtils.h"
+#include "RHIGPUReadback.h"
 
 #include "Serialization/MemoryReader.h"
 
 namespace NNX
 {
 
+struct FMLInferenceModelRDG::FReadbackEntry
+{
+	TUniquePtr<FRHIGPUBufferReadback>	RHI;
+	void*					CpuMemory;
+	size_t					Offset;
+	size_t					Size;
+};
+
 //
 //
 //
@@ -168,12 +177,12 @@ FMLInferenceModelRDG::FMLInferenceModelRDG()
 	: FMLInferenceModel(EMLInferenceModelType::RDG)
 	, bUseManualTransitions(false)
 {
-	Readback.RHI = new FRHIGPUBufferReadback("FMLTensorReadback");
+	
 }
 
 FMLInferenceModelRDG::~FMLInferenceModelRDG()
 {
-	delete Readback.RHI;
+	
 }
 
 //
@@ -252,6 +261,9 @@ int FMLInferenceModelRDG::Run(TConstArrayView<FMLTensorBinding> InInputBindings,
 		UE_LOG(LogNNX, Error, TEXT("Run(): Input shapes are not set, please call SetInputTensorShapes."));
 		return -1;
 	}
+
+	// TODO ok if pending?
+	Readbacks.Empty();
 	
 	int Res = 0;
 	FEvent* Signal = FGenericPlatformProcess::GetSynchEventFromPool(false);
@@ -279,7 +291,7 @@ int FMLInferenceModelRDG::Run(TConstArrayView<FMLTensorBinding> InInputBindings,
 				RHICmdList.BlockUntilGPUIdle();
 
 				// Process readback
-				{
+				for (const FReadbackEntry& Readback : Readbacks) {
 					const void* BuffData = Readback.RHI->Lock(Readback.Size);
 					check(BuffData);
 					FMemory::Memcpy(Readback.CpuMemory, BuffData, Readback.Size);
@@ -488,6 +500,8 @@ void FMLInferenceModelRDG::AddTensorReadbacks_RenderThread(FRDGBuilder& GraphBui
 	TConstArrayView<FTensorRDG> InTensorRDGs, TConstArrayView<FMLTensorBinding> InBindings)
 {
 	check(InTensorRDGs.Num() == InBindings.Num());
+
+	check(Readbacks.IsEmpty());
 	
 	for (int32 Idx = 0; Idx < InReadbackIndices.Num(); ++Idx)
 	{
@@ -520,6 +534,8 @@ void FMLInferenceModelRDG::AddTensorReadbacks_RenderThread(FRDGBuilder& GraphBui
 					RHICmdList.SubmitCommandsHint();
 				}
 
+				FReadbackEntry& Readback = Readbacks.Add_GetRef({});
+				Readback.RHI = MakeUnique<FRHIGPUBufferReadback>(FName(TEXT("FMLTensorReadback_") + TensorRDG.GetName()));
 				Readback.RHI->EnqueueCopy(RHICmdList, OutputBuffer, TensorRDG.GetDataSize());
 				Readback.CpuMemory = Binding.CpuMemory;
 				Readback.Offset = 0;
@@ -529,51 +545,4 @@ void FMLInferenceModelRDG::AddTensorReadbacks_RenderThread(FRDGBuilder& GraphBui
 	}
 }
 
-
-//
-//
-//
-//void FMLInferenceModelRDG::AddTensorUpload_RenderThread(FRDGBuilder& GraphBuilder, FRDGBufferRef TensorBuffer, const FMLTensorBinding& InTensorBinding, const FMLTensorDesc& TensorDesc)
-//{
-//	FMLTensorUploadParameters* TensorUploadParams = GraphBuilder.AllocParameters<FMLTensorUploadParameters>();
-//
-//	TensorUploadParams->Buffer = TensorBuffer;
-//
-//	GraphBuilder.AddPass(
-//		RDG_EVENT_NAME("NNXDmlTensorUpload"),
-//		TensorUploadParams,
-//		ERDGPassFlags::Copy | ERDGPassFlags::NeverCull,
-//		[InTensorBinding, TensorDesc, TensorUploadParams](FRHICommandListImmediate& RHICmdList)
-//		{
-//			// Copy input					
-//			void* BuffData = RHICmdList.LockBuffer(TensorUploadParams->Buffer->GetRHI(), 0, TensorDesc.DataSize, RLM_WriteOnly);
-//			FMemory::Memcpy(BuffData, InTensorBinding.CpuMemory, TensorDesc.DataSize);
-//			RHICmdList.UnlockBuffer(TensorUploadParams->Buffer->GetRHI());
-//		}
-//	);
-//}
-
-//
-//
-//
-//void FMLInferenceModelRDG::AddTensorReadback_RenderThread(FRDGBuilder& GraphBuilder, const FMLTensorBinding& InTensorBinding, const FMLTensorDesc& TensorDesc)
-//{	
-//	FMLTensorReadbackParameters* TensorReadbackParams = GraphBuilder.AllocParameters<FMLTensorReadbackParameters>();
-//
-//	TensorReadbackParams->Buffer = TensorBuffer;
-//
-//	GraphBuilder.AddPass(
-//		RDG_EVENT_NAME("NNXDmlTensorReadback"),
-//		TensorReadbackParams,
-//		ERDGPassFlags::Copy | ERDGPassFlags::NeverCull,
-//		[InTensorBinding, TensorDesc, TensorReadbackParams](FRHICommandListImmediate& RHICmdList)
-//		{
-//			// Copy input					
-//			void* BuffData = RHICmdList.LockBuffer(TensorReadbackParams->Buffer->GetRHI(), 0, TensorDesc.DataSize, RLM_WriteOnly);
-//			FMemory::Memcpy(BuffData, InTensorBinding.CpuMemory, TensorDesc.DataSize);
-//			RHICmdList.UnlockBuffer(TensorReadbackParams->Buffer->GetRHI());
-//		}
-//	);
-//}
-
 } // NNX
diff --git a/Engine/Plugins/Experimental/NNX/Source/NNXRuntimeRDG/Private/NNXRuntimeRDG.h b/Engine/Plugins/Experimental/NNX/Source/NNXRuntimeRDG/Private/NNXRuntimeRDG.h
--- a/Engine/Plugins/Experimental/NNX/Source/NNXRuntimeRDG/Private/NNXRuntimeRDG.h
+++ b/Engine/Plugins/Experimental/NNX/Source/NNXRuntimeRDG/Private/NNXRuntimeRDG.h
@@ -9,7 +9,6 @@
 #include "NNXRuntime.h"
 #include "NNXRuntimeFormat.h"
 #include "RenderGraphResources.h"
-#include "RHIGPUReadback.h"
 #include "Serialization/MemoryReader.h"
 #include "ShaderParameterUtils.h"
 
@@ -74,13 +73,8 @@ using FIntArray = TArray<int32, TInlineAllocator<16>>;
  */
 class FMLInferenceModelRDG : public FMLInferenceModel
 {
-	struct FReadbackEntry
-	{
-		FRHIGPUBufferReadback*	RHI;
-		void*					CpuMemory;
-		size_t					Offset;
-		size_t					Size;
-	};
+
+	struct FReadbackEntry;
 
 public:
 
@@ -120,7 +114,7 @@ protected:
 	//RDG Tensors
 	FTensorRDGArray				AllTensorRDGs;
 	
-	FReadbackEntry				Readback;
+	TArray<FReadbackEntry> 		Readbacks;
 	bool						bUseManualTransitions;
 };
 
