NNE - Adding model import, caching and cook functionality

diff --git a/Engine/Plugins/Experimental/NNX/NNX.uplugin b/Engine/Plugins/Experimental/NNX/NNX.uplugin
--- a/Engine/Plugins/Experimental/NNX/NNX.uplugin
+++ b/Engine/Plugins/Experimental/NNX/NNX.uplugin
@@ -78,6 +78,12 @@
 			"Type": "RuntimeAndProgram",
 			"LoadingPhase": "Default",
 			"PlatformAllowList": [ "Win64", "Linux", "Mac" ]
+		},
+		{
+			"Name": "NNEEditorTools",
+			"Type": "Editor",
+			"LoadingPhase": "Default",
+			"PlatformAllowList": [ "Win64", "Linux", "Mac" ]
 		}
 	]
 }
diff --git a/Engine/Plugins/Experimental/NNX/Source/NNEEditorTools/NNEEditorTools.Build.cs b/Engine/Plugins/Experimental/NNX/Source/NNEEditorTools/NNEEditorTools.Build.cs
new file mode 100644
--- /dev/null
+++ b/Engine/Plugins/Experimental/NNX/Source/NNEEditorTools/NNEEditorTools.Build.cs
@@ -0,0 +1,31 @@
+// Copyright Epic Games, Inc. All Rights Reserved.
+
+using UnrealBuildTool;
+
+public class NNEEditorTools : ModuleRules
+{
+	public NNEEditorTools(ReadOnlyTargetRules Target) : base(Target)
+	{
+		PCHUsage = PCHUsageMode.UseExplicitOrSharedPCHs;
+
+		PublicDependencyModuleNames.AddRange
+			(
+			new string[] {
+				"Core",
+				"CoreUObject",
+				"Engine",
+				"NNXCore"
+			}
+		);
+
+		if (Target.Type == TargetType.Editor)
+		{
+			PrivateDependencyModuleNames.AddRange(
+				new string[]
+				{
+					"UnrealEd"
+				}
+			);
+		}
+	}
+}
diff --git a/Engine/Plugins/Experimental/NNX/Source/NNEEditorTools/Private/NNEEditorTools.cpp b/Engine/Plugins/Experimental/NNX/Source/NNEEditorTools/Private/NNEEditorTools.cpp
new file mode 100644
--- /dev/null
+++ b/Engine/Plugins/Experimental/NNX/Source/NNEEditorTools/Private/NNEEditorTools.cpp
@@ -0,0 +1,19 @@
+// Copyright Epic Games, Inc. All Rights Reserved.
+
+#include "Modules/ModuleManager.h"
+
+class FNNEEditorToolsModule : public IModuleInterface
+{
+public:
+
+	virtual void StartupModule() override
+	{
+	}
+
+	virtual void ShutdownModule() override
+	{
+	}
+
+};
+
+IMPLEMENT_MODULE(FNNEEditorToolsModule, NNEEditorTools)
\ No newline at end of file
diff --git a/Engine/Plugins/Experimental/NNX/Source/NNEEditorTools/Private/NNEModelDataFactory.cpp b/Engine/Plugins/Experimental/NNX/Source/NNEEditorTools/Private/NNEModelDataFactory.cpp
new file mode 100644
--- /dev/null
+++ b/Engine/Plugins/Experimental/NNX/Source/NNEEditorTools/Private/NNEModelDataFactory.cpp
@@ -0,0 +1,31 @@
+// Copyright Epic Games, Inc. All Rights Reserved.
+
+#include "NNEModelDataFactory.h"
+#include "NNEModelData.h"
+
+UNNEModelDataFactory::UNNEModelDataFactory(const FObjectInitializer& ObjectInitializer) : UFactory(ObjectInitializer)
+{
+	bCreateNew = false;
+	bEditorImport = true;
+	SupportedClass = UNNEModelData::StaticClass();
+	ImportPriority = DefaultImportPriority;
+	Formats.Add("onnx;Open Neural Network Exchange Format");
+}
+
+UObject* UNNEModelDataFactory::FactoryCreateBinary(UClass* Class, UObject* InParent, FName Name, EObjectFlags Flags, UObject* Context, const TCHAR * Type, const uint8 *& Buffer, const uint8 * BufferEnd, FFeedbackContext* Warn)
+{
+	if (!Type || !Buffer || !BufferEnd || BufferEnd - Buffer <= 0)
+	{
+		return nullptr;
+	}
+
+	UNNEModelData* Result = NewObject<UNNEModelData>(InParent, Class, Name, Flags);
+	Result->Init(Type, Buffer, BufferEnd);
+
+	return Result;
+}
+
+bool UNNEModelDataFactory::FactoryCanImport(const FString & Filename)
+{
+	return Filename.EndsWith(FString("onnx"));
+}
\ No newline at end of file
diff --git a/Engine/Plugins/Experimental/NNX/Source/NNEEditorTools/Public/NNEModelDataFactory.h b/Engine/Plugins/Experimental/NNX/Source/NNEEditorTools/Public/NNEModelDataFactory.h
new file mode 100644
--- /dev/null
+++ b/Engine/Plugins/Experimental/NNX/Source/NNEEditorTools/Public/NNEModelDataFactory.h
@@ -0,0 +1,24 @@
+// Copyright Epic Games, Inc. All Rights Reserved.
+
+#pragma once
+
+#include "CoreMinimal.h"
+#include "Factories/Factory.h"
+
+#include "NNEModelDataFactory.generated.h"
+
+UCLASS()
+class NNEEDITORTOOLS_API UNNEModelDataFactory : public UFactory
+{
+	GENERATED_BODY()
+
+public:
+
+	UNNEModelDataFactory(const FObjectInitializer& ObjectInitializer);
+
+public:
+
+	virtual UObject* FactoryCreateBinary(UClass* Class, UObject* InParent, FName Name, EObjectFlags Flags, UObject* Context, const TCHAR* Type, const uint8*& Buffer, const uint8* BufferEnd, FFeedbackContext* Warn) override;
+	virtual bool FactoryCanImport(const FString& Filename) override;
+
+};
\ No newline at end of file
diff --git a/Engine/Plugins/Experimental/NNX/Source/NNXCore/NNXCore.Build.cs b/Engine/Plugins/Experimental/NNX/Source/NNXCore/NNXCore.Build.cs
--- a/Engine/Plugins/Experimental/NNX/Source/NNXCore/NNXCore.Build.cs
+++ b/Engine/Plugins/Experimental/NNX/Source/NNXCore/NNXCore.Build.cs
@@ -14,12 +14,22 @@ namespace UnrealBuildTool.Rules
 
             PrivateDependencyModuleNames.AddRange(
                 new string[]
-                {
+				{
 					"Core",
 					"CoreUObject",
 					"Engine",
 				}
 			);
-        }
+
+			if (Target.Type == TargetType.Editor)
+			{
+				PrivateDependencyModuleNames.AddRange(
+					new string[]
+					{
+						"DerivedDataCache"
+					}
+				);
+			}
+		}
     }
 }
diff --git a/Engine/Plugins/Experimental/NNX/Source/NNXCore/Private/NNEModel.cpp b/Engine/Plugins/Experimental/NNX/Source/NNXCore/Private/NNEModel.cpp
new file mode 100644
--- /dev/null
+++ b/Engine/Plugins/Experimental/NNX/Source/NNXCore/Private/NNEModel.cpp
@@ -0,0 +1,171 @@
+// Copyright Epic Games, Inc. All Rights Reserved.
+
+#include "NNEModel.h"
+
+#include "NNXCore.h"
+#include "NNXRuntimeFormat.h"
+#include "NNXInferenceModel.h"
+
+#include "GenericPlatform/GenericPlatformProcess.h"
+
+UNNEModel* UNNEModel::Create(UObject* Parent, FString RuntimeName, UNNEModelData* ModelData)
+{
+	UNNEModel* Result = NewObject<UNNEModel>(Parent);
+	if (Result && Result->Load(RuntimeName, ModelData))
+	{
+		return Result;
+	}
+	return nullptr;
+}
+
+bool UNNEModel::Load(FString RuntimeName, UNNEModelData* ModelData)
+{
+	using namespace NNX;
+
+	IRuntime* Runtime = GetRuntime(RuntimeName);
+	if (!Runtime)
+	{
+		UE_LOG(LogNNX, Error, TEXT("UNNEModel: No runtime '%s' found. Valid runtimes are: "), *RuntimeName);
+		TArray<NNX::IRuntime*> Runtimes = GetAllRuntimes();
+		for (int i = 0; i < Runtimes.Num(); i++)
+		{
+			UE_LOG(LogNNX, Error, TEXT("- %s"), *Runtimes[i]->GetRuntimeName());
+		}
+		return false;
+	}
+
+	if (!ModelData)
+	{
+		UE_LOG(LogNNX, Error, TEXT("UNNEModel: Valid model data required to load the model"));
+		return false;
+	}
+
+	TConstArrayView<uint8> Data = ModelData->GetModelData(RuntimeName);
+	if (Data.Num() < 1)
+	{
+		UE_LOG(LogNNX, Error, TEXT("UNNEModel: No model data for %s found"), *RuntimeName);
+		return false;
+	}
+
+	FNNIModelRaw ModelRaw;
+	ModelRaw.Data = MoveTemp(Data);
+	ModelRaw.Format = ENNXInferenceFormat::ONNX;
+	UMLInferenceModel* TempModelData = UMLInferenceModel::CreateFromFormatDesc(ModelRaw);
+	if (!TempModelData)
+	{
+		return false;
+	}
+
+	Model = TSharedPtr<NNX::FMLInferenceModel>(Runtime->CreateInferenceModel(TempModelData));
+
+	return Model.IsValid();
+}
+
+int32 UNNEModel::NumInputs()
+{
+	if (!Model.IsValid())
+	{
+		return 0;
+	}
+	return Model->GetInputTensorDescs().Num();
+}
+
+
+int32 UNNEModel::NumOutputs()
+{
+	if (!Model.IsValid())
+	{
+		return 0;
+	}
+	return Model->GetOutputTensorDescs().Num();
+}
+
+
+TArray<int32> UNNEModel::GetInputShapes(int32 Index)
+{
+	if (!Model.IsValid())
+	{
+		return TArray<int32>();
+	}
+
+	TConstArrayView<NNX::FTensorDesc> Desc = Model->GetInputTensorDescs();
+	if (Index < 0 || Index >= Desc.Num())
+	{
+		return TArray<int32>();
+	}
+
+	return TArray<int32>(Desc[Index].GetShape().Data);
+}
+
+TArray<int32> UNNEModel::GetOutputShapes(int32 Index)
+{
+	if (!Model.IsValid())
+	{
+		return TArray<int32>();
+	}
+
+	TConstArrayView<NNX::FTensorDesc> Desc = Model->GetOutputTensorDescs();
+	if (Index < 0 || Index >= Desc.Num())
+	{
+		return TArray<int32>();
+	}
+
+	return TArray<int32>(Desc[Index].GetShape().Data);
+}
+
+bool UNNEModel::SetInputOutput(const TArray<FNNETensor>& Input, UPARAM(ref) TArray<FNNETensor>& Output)
+{
+	if (!Model.IsValid())
+	{
+		return false;
+	}
+
+	InputBindings.Empty();
+	OutputBindings.Empty();
+
+	TConstArrayView<NNX::FTensorDesc> InputDescs = Model->GetInputTensorDescs();
+	if (InputDescs.Num() != Input.Num())
+	{
+		UE_LOG(LogNNX, Error, TEXT("UNNEModel: Invalid number of input tensors provided"));
+		return false;
+	}
+
+	TConstArrayView<NNX::FTensorDesc> OutputDescs = Model->GetOutputTensorDescs();
+	if (OutputDescs.Num() != Output.Num())
+	{
+		UE_LOG(LogNNX, Error, TEXT("UNNEModel: Invalid number of output tensors provided"));
+		return false;
+	}
+
+	TArray<NNX::FTensorShape> InputShapes;
+	for (int i = 0; i < Input.Num(); i++)
+	{
+		InputBindings.Add(NNX::FMLTensorBinding::FromCPU((void*)(Input[i].Data.GetData()), 4 * Input[i].Data.Num()));
+
+		InputShapes.Add(NNX::FTensorShape::MakeFromSymbolic(NNX::FSymbolicTensorShape::Make(Input[i].Shape)));
+	}
+
+	for (int i = 0; i < Input.Num(); i++)
+	{
+		OutputBindings.Add(NNX::FMLTensorBinding::FromCPU((void*)(Output[i].Data.GetData()), 4 * Output[i].Data.Num()));
+	}
+
+	Model->SetInputTensorShapes(InputShapes);
+
+	return true;
+}
+
+bool UNNEModel::Run()
+{
+	if (!Model.IsValid())
+	{
+		return false;
+	}
+
+	return Model->Run(InputBindings, OutputBindings) == 0;
+}
+
+TSharedPtr<NNX::FMLInferenceModel> UNNEModel::GetModel()
+{
+	return Model;
+}
\ No newline at end of file
diff --git a/Engine/Plugins/Experimental/NNX/Source/NNXCore/Private/NNEModelData.cpp b/Engine/Plugins/Experimental/NNX/Source/NNXCore/Private/NNEModelData.cpp
new file mode 100644
--- /dev/null
+++ b/Engine/Plugins/Experimental/NNX/Source/NNXCore/Private/NNEModelData.cpp
@@ -0,0 +1,237 @@
+// Copyright Epic Games, Inc. All Rights Reserved.
+
+#include "NNEModelData.h"
+
+#include "NNXCore.h"
+#include "NNXRuntimeFormat.h"
+#include "NNXModelOptimizerInterface.h"
+#include "NNECoreAttributeMap.h"
+#include "Serialization/CustomVersion.h"
+
+#if WITH_EDITOR
+#include "Misc/Guid.h"
+#include "Containers/StringFwd.h"
+#include "DerivedDataCacheKey.h"
+#include "DerivedDataCache.h"
+#include "Memory/SharedBuffer.h"
+#include "DerivedDataRequestOwner.h"
+#include "Internationalization/TextLocalizationResource.h"
+#endif
+
+const FGuid UNNEModelData::GUID(0x9513202e, 0xeba1b279, 0xf17fe5ba, 0xab90c3f2);
+FCustomVersionRegistration NNEModelDataVersion(UNNEModelData::GUID, 0, TEXT("NNEModelDataVersion"));
+
+#if WITH_EDITOR
+
+inline UE::DerivedData::FCacheKey CreateCacheKey(const FGuid& FileDataId, const FString& RuntimeName)
+{
+	FString GuidString = FileDataId.ToString(EGuidFormats::Digits);
+	return { UE::DerivedData::FCacheBucket(FWideStringView(*GuidString)), FIoHash::HashBuffer(MakeMemoryView(FTCHARToUTF8(RuntimeName))) };
+}
+
+inline FSharedBuffer GetFromDDC(const FGuid& FileDataId, const FString& RuntimeName)
+{
+	UE::DerivedData::FCacheGetValueRequest GetRequest;
+	GetRequest.Name = FString("Get-") + RuntimeName + FString("-") + FileDataId.ToString(EGuidFormats::Digits);
+	GetRequest.Key = CreateCacheKey(FileDataId,  RuntimeName);
+	FSharedBuffer RawDerivedData;
+	UE::DerivedData::FRequestOwner BlockingGetOwner(UE::DerivedData::EPriority::Blocking);
+	UE::DerivedData::GetCache().GetValue({ GetRequest }, BlockingGetOwner, [&RawDerivedData](UE::DerivedData::FCacheGetValueResponse&& Response)
+		{
+			RawDerivedData = Response.Value.GetData().Decompress();
+		});
+	BlockingGetOwner.Wait();
+	return RawDerivedData;
+}
+
+inline void PutIntoDDC(const FGuid& FileDataId, const FString& RuntimeName, FSharedBuffer& Data)
+{
+	UE::DerivedData::FCachePutValueRequest PutRequest;
+	PutRequest.Name = FString("Put-") + RuntimeName + FString("-") + FileDataId.ToString(EGuidFormats::Digits);
+	PutRequest.Key = CreateCacheKey(FileDataId, RuntimeName);
+	PutRequest.Value = UE::DerivedData::FValue::Compress(Data);
+	UE::DerivedData::FRequestOwner BlockingPutOwner(UE::DerivedData::EPriority::Blocking);
+	UE::DerivedData::GetCache().PutValue({ PutRequest }, BlockingPutOwner);
+	BlockingPutOwner.Wait();
+}
+
+#endif
+
+inline TArray<uint8> Create(const FString& RuntimeName, const TArray<uint8>& FileData)
+{
+	using namespace NNX;
+
+	// Try to optimize the model
+	IRuntime* Runtime = GetRuntime(RuntimeName);
+	if (!Runtime)
+	{
+		UE_LOG(LogNNX, Error, TEXT("UNNEModelData: No runtime '%s' found. Valid runtimes are: "), *RuntimeName);
+		TArray<NNX::IRuntime*> Runtimes = GetAllRuntimes();
+		for (int i = 0; i < Runtimes.Num(); i++)
+		{
+			UE_LOG(LogNNX, Error, TEXT("- %s"), *Runtimes[i]->GetRuntimeName());
+		}
+		return {};
+	}
+
+	TUniquePtr<IModelOptimizer> Optimizer = Runtime->CreateModelOptimizer();
+	if (!Optimizer.IsValid())
+	{
+		UE_LOG(LogNNX, Error, TEXT("UNNEModelData: %s cannot create an optimizer"), *RuntimeName);
+		return {};
+	}
+
+	FNNIModelRaw InputModel;
+	InputModel.Data = FileData;
+	InputModel.Format = ENNXInferenceFormat::ONNX;
+	FNNIModelRaw OutputModel;
+	FOptimizerOptionsMap Options;
+	if (!Optimizer->Optimize(InputModel, OutputModel, Options))
+	{
+		UE_LOG(LogNNX, Error, TEXT("UNNEModelData: Optimization failed"));
+		return {};
+	}
+
+	return MoveTemp(OutputModel.Data);
+}
+
+void UNNEModelData::Init(const TCHAR* Type, const uint8*& Buffer, const uint8* BufferEnd)
+{
+	check(Type)
+	check(Buffer)
+	check(BufferEnd)
+	check(BufferEnd - Buffer > 0)
+
+	FileType = FString(Type);
+	FileData = TArray<uint8>(Buffer, BufferEnd - Buffer);
+	FPlatformMisc::CreateGuid(FileDataId);
+	ModelData.Empty();
+}
+
+TConstArrayView<uint8> UNNEModelData::GetModelData(const FString& RuntimeName)
+{
+#if WITH_EDITOR
+
+	using namespace NNX;
+
+	// Check if we have a local cache hit
+	TArray<uint8>* LocalData = ModelData.Find(RuntimeName);
+	if (LocalData)
+	{
+		UE_LOG(LogNNX, Display, TEXT("UNNEModelData: Using locally cached model data"));
+		return TConstArrayView<uint8>(LocalData->GetData(), LocalData->Num());
+	}
+
+	// Check if we have a remote cache hit
+	FSharedBuffer RemoteData = GetFromDDC(FileDataId, RuntimeName);
+	if (RemoteData.GetSize() > 0)
+	{
+		ModelData.Add(RuntimeName, TArray<uint8>((uint8*)RemoteData.GetData(), RemoteData.GetSize()));
+
+		UE_LOG(LogNNX, Display, TEXT("UNNEModelData: Using remote cached model data"));
+
+		TArray<uint8>* CachedRemoteData = ModelData.Find(RuntimeName);
+		return TConstArrayView<uint8>(CachedRemoteData->GetData(), CachedRemoteData->Num());
+	}
+
+	// Try to create the model
+	TArray<uint8> CreatedData = Create(RuntimeName, FileData);
+	if (CreatedData.Num() < 1)
+	{
+		return {};
+	}
+
+	// Cache the model and put it into ddc
+	UE_LOG(LogNNX, Display, TEXT("UNNEModelData: Caching created model data"));
+	ModelData.Add(RuntimeName, CreatedData);
+
+	FSharedBuffer SharedBuffer = MakeSharedBufferFromArray(MoveTemp(CreatedData));
+	PutIntoDDC(FileDataId, RuntimeName, SharedBuffer);
+
+	TArray<uint8>* CachedCreatedData = ModelData.Find(RuntimeName);
+	return TConstArrayView<uint8>(CachedCreatedData->GetData(), CachedCreatedData->Num());
+
+#else
+
+	TArray<uint8>* Data = ModelData.Find(RuntimeName);
+	if (Data)
+	{
+		UE_LOG(LogTemp, Display, TEXT("UNNEModelData: Using cooked model data"));
+		return TConstArrayView<uint8>(Data->GetData(), Data->Num());
+	}
+	return {};
+
+#endif
+}
+
+void UNNEModelData::Serialize(FArchive& Ar)
+{
+	using namespace NNX;
+
+	// Store the asset version (no effect in load)
+	Ar.UsingCustomVersion(UNNEModelData::GUID);
+
+	// Only one version is supported for now
+	check(Ar.CustomVer(UNNEModelData::GUID) == 0);
+
+	// Recreate each model data when cooking
+	if (Ar.IsCooking() && Ar.IsSaving())
+	{
+		ModelData.Empty();
+		TArray<NNX::IRuntime*> Runtimes = GetAllRuntimes();
+		for (int i = 0; i < Runtimes.Num(); i++)
+		{
+			TArray<uint8> CreatedData = Create(Runtimes[i]->GetRuntimeName(), FileData);
+			if (CreatedData.Num() > 0)
+			{
+				ModelData.Add(Runtimes[i]->GetRuntimeName(), CreatedData);
+#if WITH_EDITOR
+				FSharedBuffer SharedBuffer = MakeSharedBufferFromArray(MoveTemp(CreatedData));
+				PutIntoDDC(FileDataId, Runtimes[i]->GetRuntimeName(), SharedBuffer);
+#endif
+			}
+		}
+
+		// Dummy data for fields not required in the game
+		TArray<uint8> EmptyData;
+		TArray<FString> RuntimeNames;
+		ModelData.GetKeys(RuntimeNames);
+		int32 NumItems = RuntimeNames.Num();
+
+		Ar << FileType;
+		Ar << EmptyData;
+		Ar << FileDataId;
+		Ar << NumItems;
+
+		for (int i = 0; i < NumItems; i++)
+		{
+			Ar << RuntimeNames[i];
+			Ar << ModelData[RuntimeNames[i]];
+		}
+
+		UE_LOG(LogTemp, Display, TEXT("UNNEModelData: Serialized data of %d runtimes"), NumItems);
+	}
+	else
+	{
+		int32 NumItems = 0;
+
+		Ar << FileType;
+		Ar << FileData;
+		Ar << FileDataId;
+		Ar << NumItems;
+
+		if (Ar.IsLoading())
+		{
+			for (int i = 0; i < NumItems; i++)
+			{
+				FString Name;
+				Ar << Name;
+				TArray<uint8> Data;
+				Ar << Data;
+				ModelData.Add(Name, MoveTemp(Data));
+			}
+
+			UE_LOG(LogTemp, Display, TEXT("UNNEModelData: Deserialized data of %d runtimes"), NumItems);
+		}
+	}
+}
\ No newline at end of file
diff --git a/Engine/Plugins/Experimental/NNX/Source/NNXCore/Public/NNEModel.h b/Engine/Plugins/Experimental/NNX/Source/NNXCore/Public/NNEModel.h
new file mode 100644
--- /dev/null
+++ b/Engine/Plugins/Experimental/NNX/Source/NNXCore/Public/NNEModel.h
@@ -0,0 +1,92 @@
+// Copyright Epic Games, Inc. All Rights Reserved.
+
+#pragma once
+
+#include "CoreMinimal.h"
+
+#include "NNXRuntime.h"
+#include "NNEModelData.h"
+
+#include "NNEModel.generated.h"
+
+USTRUCT(BlueprintType, Category = "NNE - Neural Network Engine")
+struct FNNETensor
+{
+	GENERATED_BODY()
+
+public:
+
+	UPROPERTY(BlueprintReadWrite, Category = "NNE - Neural Network Engine")
+	TArray<int32> Shape = TArray<int32>();
+
+	UPROPERTY(BlueprintReadWrite, Category = "NNE - Neural Network Engine")
+	TArray<float> Data = TArray<float>();
+};
+
+UCLASS(BlueprintType, Category = "NNE - Neural Network Engine")
+class NNXCORE_API UNNEModel : public UObject
+{
+	GENERATED_BODY()
+
+public:
+
+	/**
+	 * Creates a new model.
+	 */
+	UFUNCTION(BlueprintCallable, Category = "NNE - Neural Network Engine")
+	static UNNEModel* Create(UObject* Parent, FString RuntimeName, UNNEModelData* ModelData);
+
+	/**
+	 * Loads the model from model data for a given runtime.
+	 */
+	UFUNCTION(BlueprintCallable, Category = "NNE - Neural Network Engine")
+	virtual bool Load(FString RuntimeName, UNNEModelData* ModelData);
+
+public:
+
+	UFUNCTION(BlueprintCallable, Category = "NNE - Neural Network Engine")
+	virtual int32 NumInputs();
+
+	UFUNCTION(BlueprintCallable, Category = "NNE - Neural Network Engine")
+	virtual int32 NumOutputs();
+
+	UFUNCTION(BlueprintCallable, Category = "NNE - Neural Network Engine")
+	virtual TArray<int32> GetInputShapes(int32 Index);
+
+	UFUNCTION(BlueprintCallable, Category = "NNE - Neural Network Engine")
+	virtual TArray<int32> GetOutputShapes(int32 Index);
+
+public:
+
+	/**
+	 * Sets the input and output tensors.
+	 * If changes are made to the shape or data sizes inside the passed tensors, this function must be called again. 
+	 */
+	UFUNCTION(BlueprintCallable, Category = "NNE - Neural Network Engine")
+	virtual bool SetInputOutput(const TArray<FNNETensor>& Input, UPARAM(ref) TArray<FNNETensor>& Output);
+	
+public:
+
+	/**
+	 * Synchroneous call to run the model.
+	 * Requires the runtime to be selected with SetRuntime(...) first.
+	 * If an RDG runtime is selected, data will be up and downloaded to and from the GPU.
+	 * Upon successful return of this function, the data inside the output tensors passed to SetInputOutput will contain the resulting data.
+	 */
+	UFUNCTION(BlueprintCallable, Category = "NNE - Neural Network Engine")
+	virtual bool Run();
+	
+public:
+	
+	/**
+	 * Get the model to gain full access including running from e.g. the render thread.
+	 */
+	TSharedPtr<NNX::FMLInferenceModel> GetModel();
+	
+private:
+
+	TSharedPtr<NNX::FMLInferenceModel> Model;
+
+	TArray<NNX::FMLTensorBinding> InputBindings;
+	TArray<NNX::FMLTensorBinding> OutputBindings;
+};
diff --git a/Engine/Plugins/Experimental/NNX/Source/NNXCore/Public/NNEModelData.h b/Engine/Plugins/Experimental/NNX/Source/NNXCore/Public/NNEModelData.h
new file mode 100644
--- /dev/null
+++ b/Engine/Plugins/Experimental/NNX/Source/NNXCore/Public/NNEModelData.h
@@ -0,0 +1,67 @@
+// Copyright Epic Games, Inc. All Rights Reserved.
+
+#pragma once
+
+#include "CoreMinimal.h"
+
+#include "NNEModelData.generated.h"
+
+UCLASS(BlueprintType, Category = "NNE - Neural Network Engine")
+class NNXCORE_API UNNEModelData : public UObject
+{
+	GENERATED_BODY()
+
+public:
+
+	/**
+	 * Function to initialize the data.
+	 * Called by the NNEModelDataFactory.
+	 * Type corresponds to the file extension.
+	 */
+	void Init(const TCHAR* Type, const uint8*& Buffer, const uint8* BufferEnd);
+
+	/**
+	 * The function returns the cached (editor) or cooked (game) optimized model data for a given runtime.
+	 * In editor, the function will create the model data with the passed runtime in case it has not been cached yet.
+	 * Returns an empty view in case of a failure.
+	 */
+	TConstArrayView<uint8> GetModelData(const FString& RuntimeName);
+
+public:
+
+	virtual void Serialize(FArchive& Ar) override;
+
+public:
+
+	/**
+	 * A GUID used for versioning.
+	 */
+	const static FGuid GUID;
+
+private:
+
+	/**
+	 * The file type passed by the factory when importing a model.
+	 * Corresponds to the file extension.
+	 */
+	FString FileType;
+
+	/**
+	 * Raw binary file data of the imported model.
+	 */
+	TArray<uint8> FileData;
+
+	/**
+	 * Guid that uniquely identifies this model.
+	 * This is used to cache optimized models in the editor.
+	 */
+	FGuid FileDataId;
+
+private:
+
+	/**
+	 * The processed / optimized model data for the different runtimes.
+	 */
+	TMap<FString, TArray<uint8>> ModelData;
+
+};
