NNI Refactor test framework to allow test definition to provide the input data

Problem:
* Test description need to be able to define the input data for model weight, input and output data to cover more case or do it more explicitely

Solution:
* Refactor the test framework to allow it to create the test data at the test definition level. Runtime then do a memcopy of the data before running.
   This also have the advantage of preparing the data only once for all runtime engine, speeding our test suite by around 20%.

Tests:
* All NNX tests pass under win64 editor

diff --git a/Engine/Plugins/Experimental/NNX/Source/NNXCore/Public/NNXTypes.h b/Engine/Plugins/Experimental/NNX/Source/NNXCore/Public/NNXTypes.h
--- a/Engine/Plugins/Experimental/NNX/Source/NNXCore/Public/NNXTypes.h
+++ b/Engine/Plugins/Experimental/NNX/Source/NNXCore/Public/NNXTypes.h
@@ -42,6 +42,24 @@ template <typename DimType, class FinalType> struct TTensorShapeBase
 		return Data.Num();
 	}
 
+	inline bool operator==(const TTensorShapeBase<DimType,FinalType>& OtherShape) const
+	{
+		check(Rank() <= MaxRank);
+
+		if (Rank() != OtherShape.Rank())
+			return false;
+
+		for (int32 i = 0; i < Rank(); ++i)
+		{
+			if (Data[i] != OtherShape.Data[i])
+			return false;
+		}
+
+		return true;
+	}
+	
+	inline bool operator!=(const TTensorShapeBase<DimType, FinalType>& OtherShape) const { return !(*this == OtherShape); }
+
 	static FinalType Make(TConstArrayView<DimType> Data)
 	{
 		if (Data.Num() > MaxRank)
@@ -85,7 +103,7 @@ struct FTensorShape : TTensorShapeBase<uint32, FTensorShape>
 		
 		return Result;
 	}
-	
+
 	bool IsCompatibleWith(FSymbolicTensorShape SymbolicShape) const
 	{
 		if (Rank() != SymbolicShape.Rank())
diff --git a/Engine/Plugins/Experimental/NNX/Source/NNXQA/Private/NNXQAJsonUtils.h b/Engine/Plugins/Experimental/NNX/Source/NNXQA/Private/NNXQAJsonUtils.h
--- a/Engine/Plugins/Experimental/NNX/Source/NNXQA/Private/NNXQAJsonUtils.h
+++ b/Engine/Plugins/Experimental/NNX/Source/NNXQA/Private/NNXQAJsonUtils.h
@@ -18,19 +18,24 @@ namespace Json
 	struct FTestConfigTensor : FJsonSerializable
 	{
 		TArray<int32> Shape;
+		TArray<FString> Source;
 		FString Type;
-		//TODO Extend:
-		// by adding custom data
-		// by adding custom initializer or parameter (range, distribution)
-		// by adding data type
-		// by adding source CPU or GPU resident data
-
+		//Idea: Extend with custom initializer:
+		//	FString Initializer;
+		//	JSON_SERIALIZE("initializer", ...);
+		//		if (initializer && source)
+		//			InEngineTensorData = initializer(source) <-- To add
+		//		if (!initializer && source)
+		//			InEngineTensorData = element-wise( static_cast<Type>(source[i]) ) <-- Current behavior
+		//		if (!initializer && !source)
+		//			InEngineTensorData = random from seed <-- Current behavior
+		//Idea: Extend by adding source CPU or GPU resident data
+		//	bool OnGpu;
+		//	JSON_SERIALIZE("on_gpu", OnGpu);
 		BEGIN_JSON_SERIALIZER
 			JSON_SERIALIZE_ARRAY("shape", Shape);
-			JSON_SERIALIZE("data_type", Type);
-			//JSON_SERIALIZE("on_gpu", OnGpu);
-			//JSON_SERIALIZE("range", ...);
-			//JSON_SERIALIZE("data_source", DataSource);
+			JSON_SERIALIZE("type", Type);
+			JSON_SERIALIZE_ARRAY("source", Source);
 		END_JSON_SERIALIZER
 	};
 
@@ -52,11 +57,13 @@ namespace Json
 	struct FTestConfigDataset : FJsonSerializable
 	{
 		TArray<FTestConfigTensor> Inputs;
+		TArray<FTestConfigTensor> Weights;
 		TArray<FTestConfigTensor> Outputs;
 		TArray<FTestConfigRuntime> Runtimes;
 
 		BEGIN_JSON_SERIALIZER
 			JSON_SERIALIZE_ARRAY_SERIALIZABLE("inputs", Inputs, FTestConfigTensor);
+			JSON_SERIALIZE_ARRAY_SERIALIZABLE("weights", Weights, FTestConfigTensor);
 			JSON_SERIALIZE_ARRAY_SERIALIZABLE("outputs", Outputs, FTestConfigTensor);
 			JSON_SERIALIZE_ARRAY_SERIALIZABLE("runtimes", Runtimes, FTestConfigRuntime);
 		END_JSON_SERIALIZER
diff --git a/Engine/Plugins/Experimental/NNX/Source/NNXQA/Private/NNXQAParametricTest.cpp b/Engine/Plugins/Experimental/NNX/Source/NNXQA/Private/NNXQAParametricTest.cpp
--- a/Engine/Plugins/Experimental/NNX/Source/NNXQA/Private/NNXQAParametricTest.cpp
+++ b/Engine/Plugins/Experimental/NNX/Source/NNXQA/Private/NNXQAParametricTest.cpp
@@ -2,7 +2,6 @@
 
 #include "NNXQAParametricTest.h"
 
-#include "UObject/Class.h"
 #include "NNXCore.h"
 #include "NNXTypes.h"
 #include "NNXRuntimeFormat.h"
@@ -10,9 +9,10 @@
 #include "NNXQAUtils.h"
 #include "NNXQAJsonUtils.h"
 #include "NNXModelBuilder.h"
+#include "HAL/IConsoleManager.h"
 #include "Interfaces/IPluginManager.h"
 #include "Misc/Paths.h"
-#include "HAL/IConsoleManager.h"
+#include "UObject/Class.h"
 #include "UObject/ReflectedTypeAccessors.h"
 
 #if WITH_DEV_AUTOMATION_TESTS
@@ -118,7 +118,7 @@ namespace Test
 			TestSetup.Tags = TestTarget.Tags;
 		}
 
-		static FTensorShape GetConcreteShapeFromJsonArray(const TArray<int32>& JsonShape)
+		static FTensorShape GetConcreteShapeFromJsonArray(TConstArrayView<int32> JsonShape)
 		{
 			FSymbolicTensorShape SymbolicShape = FSymbolicTensorShape::Make(JsonShape);
 			check(SymbolicShape.IsConcrete());
@@ -133,6 +133,18 @@ namespace Test
 			return (Value == INDEX_NONE) ? DefaultValue : (EMLTensorDataType)Value;
 		}
 
+		static TArray<char> GetInputTensorDataFromJson(FTensor Tensor, int TensorIndex, TConstArrayView<FString> JsonValues)
+		{
+			checkf(JsonValues.Num() == 0, TEXT("source data for tensor not supported yet"));
+			return GenerateTensorDataForTest(Tensor, TensorIndex);
+		}
+
+		static TArray<char> GetOutputTensorDataFromJson(FTensor Tensor, int TensorIndex, TConstArrayView<FString> JsonValues)
+		{
+			checkf(JsonValues.Num() == 0, TEXT("source data for tensor not supported yet"));
+			return TArray<char>();
+		}
+
 		static void ApplyDatasetConfig(FTests::FTestSetup& TestSetup, const Json::FTestConfigDataset& TestDataset, EMLTensorDataType DefaultInputType, EMLTensorDataType DefaultOutputType)
 		{
 			ApplyRuntimesConfig(TestSetup, TestDataset.Runtimes);
@@ -143,26 +155,51 @@ namespace Test
 			}
 
 			uint32 i = 0;
-			for (auto&& Tensor : TestDataset.Inputs)
+			for (auto&& JsonTensor : TestDataset.Inputs)
 			{
-				FTensorShape Shape = GetConcreteShapeFromJsonArray(Tensor.Shape);
-				EMLTensorDataType TensorType = GetTensorTypeFromJson(Tensor.Type, DefaultInputType);
-				TestSetup.Inputs.Emplace(FTensor::Make(FString::Printf(TEXT("in%d"), i++), Shape, TensorType));
+				FTensorShape Shape = GetConcreteShapeFromJsonArray(JsonTensor.Shape);
+				EMLTensorDataType TensorType = GetTensorTypeFromJson(JsonTensor.Type, DefaultInputType);
+				FTensor Tensor = FTensor::Make(FString::Printf(TEXT("input%d"), i), Shape, TensorType);
+				TArray<char> TensorData = GetInputTensorDataFromJson(Tensor, i, JsonTensor.Source);
+
+				TestSetup.Inputs.Emplace(Tensor);
+				TestSetup.InputsData.Emplace(TensorData);
+				++i;
 			}
 
-			i = 0;
-			for (auto&& Tensor : TestDataset.Outputs)
+			for (auto&& JsonTensor : TestDataset.Weights)
 			{
-				FTensorShape Shape = GetConcreteShapeFromJsonArray(Tensor.Shape);
-				EMLTensorDataType TensorType = GetTensorTypeFromJson(Tensor.Type, DefaultOutputType);
-				TestSetup.Outputs.Emplace(FTensor::Make(FString::Printf(TEXT("output%d"), i++), Shape, TensorType));
+				FTensorShape Shape = GetConcreteShapeFromJsonArray(JsonTensor.Shape);
+				EMLTensorDataType TensorType = GetTensorTypeFromJson(JsonTensor.Type, EMLTensorDataType::Float);
+				FTensor Tensor = FTensor::Make(FString::Printf(TEXT("weights%d"), i), Shape, TensorType);
+				TArray<char> TensorData = GetInputTensorDataFromJson(Tensor, i, JsonTensor.Source);
+
+				TestSetup.Weights.Emplace(Tensor);
+				TestSetup.WeightsData.Emplace(TensorData);
+				++i;
+			}
+
+			for (auto&& JsonTensor : TestDataset.Outputs)
+			{
+				FTensorShape Shape = GetConcreteShapeFromJsonArray(JsonTensor.Shape);
+				EMLTensorDataType TensorType = GetTensorTypeFromJson(JsonTensor.Type, DefaultOutputType);
+				FTensor Tensor = FTensor::Make(FString::Printf(TEXT("output%d"), i), Shape, TensorType);
+				TArray<char> TensorData = GetOutputTensorDataFromJson(Tensor, i, JsonTensor.Source);
+
+				TestSetup.Outputs.Emplace(Tensor);
+				TestSetup.OutputsData.Emplace(TensorData);
+				++i;
 			}
 			//If output is not defined it is the first input shape.
 			if (TestDataset.Outputs.Num() == 0 && TestDataset.Inputs.Num() > 0)
 			{
 				FTensorShape Shape = GetConcreteShapeFromJsonArray(TestDataset.Inputs[0].Shape);
 				EMLTensorDataType TensorType = GetTensorTypeFromJson(TestDataset.Inputs[0].Type, DefaultOutputType);
-				TestSetup.Outputs.Emplace(FTensor::Make(TEXT("output"), Shape, TensorType));
+				FTensor Tensor = FTensor::Make(FString::Printf(TEXT("output"), i), Shape, TensorType);
+				TArray<char> TensorData = GetOutputTensorDataFromJson(Tensor, i, TestDataset.Inputs[0].Source);
+				
+				TestSetup.Outputs.Emplace(Tensor);
+				TestSetup.OutputsData.Emplace(TensorData);
 			}
 		}
 
diff --git a/Engine/Plugins/Experimental/NNX/Source/NNXQA/Private/NNXQAUtils.cpp b/Engine/Plugins/Experimental/NNX/Source/NNXQA/Private/NNXQAUtils.cpp
--- a/Engine/Plugins/Experimental/NNX/Source/NNXQA/Private/NNXQAUtils.cpp
+++ b/Engine/Plugins/Experimental/NNX/Source/NNXQA/Private/NNXQAUtils.cpp
@@ -8,8 +8,6 @@
 #include "NNXInferenceModel.h"
 #include "NNXModelOptimizerInterface.h"
 
-#include <functional>
-
 namespace NNX 
 {
 namespace Test 
@@ -29,62 +27,117 @@ namespace Test
 			}
 		}
 	}
-	static void FillTensorBindings(TConstArrayView<FTensor> TensorDescs,
-		TArray<TArray<char>>& OutMemBuffers, TArray<FMLTensorBinding>& OutBindings,
-        std::function<float(EMLTensorDataType, uint32, uint32)> Initializer)
-	{
-		OutMemBuffers.Empty();
-		OutBindings.Empty();
 
-		for (int Index = 0; Index < TensorDescs.Num(); ++Index)
+	static float ElementWiseCosTensorInitializer(EMLTensorDataType DataType, uint32 ElementIndex, uint32 TensorIndex)
+	{
+		const constexpr uint32 IndexOffsetBetweenTensor = 9;
+		switch (DataType)
 		{
-			const FTensor& TensorDesc = TensorDescs[Index];
-			const uint32 NumberOfElements = TensorDesc.GetVolume();
-			const int32 ElementByteSize = TensorDesc.GetElemByteSize();
-			const uint64 BufferSize = TensorDesc.GetDataSize();
-			const EMLTensorDataType DataType = TensorDesc.GetDataType();
 
-			TArray<char> CurInputBuffer;
-			CurInputBuffer.SetNum(BufferSize);
-			char* DestinationPtr = CurInputBuffer.GetData();
-			for (uint32 i = 0; i != NumberOfElements; ++i)
+		case EMLTensorDataType::Boolean:
+			return (ElementIndex + IndexOffsetBetweenTensor * TensorIndex) % 2;
+		case EMLTensorDataType::Char:
+		case EMLTensorDataType::Int8:
+		case EMLTensorDataType::Int16:
+		case EMLTensorDataType::Int32:
+		case EMLTensorDataType::Int64:
+			return 10.0f * FMath::Cos((float)(ElementIndex + IndexOffsetBetweenTensor * TensorIndex));
+		case EMLTensorDataType::UInt8:
+		case EMLTensorDataType::UInt16:
+		case EMLTensorDataType::UInt32:
+		case EMLTensorDataType::UInt64:
+			return 10.0f * FMath::Abs(FMath::Cos((float)(ElementIndex + IndexOffsetBetweenTensor * TensorIndex)));
+		default:
+			//case EMLTensorDataType::None:
+			//case EMLTensorDataType::Half:
+			//case EMLTensorDataType::Double:
+			//case EMLTensorDataType::Float:
+			//case EMLTensorDataType::Complex64:
+			//case EMLTensorDataType::Complex128:
+			//case EMLTensorDataType::BFloat16:
+			return FMath::Cos((float)(ElementIndex + IndexOffsetBetweenTensor * TensorIndex));
+		}
+	}
+
+	TArray<char> GenerateTensorDataForTest(const FTensor& Tensor, int TensorIndex)
+	{
+		const uint32 NumberOfElements = Tensor.GetVolume();
+		const int32 ElementByteSize = Tensor.GetElemByteSize();
+		const uint64 BufferSize = Tensor.GetDataSize();
+		const EMLTensorDataType DataType = Tensor.GetDataType();
+		TArray<char> Buffer;
+			
+		Buffer.SetNum(BufferSize);
+		char* DestinationPtr = Buffer.GetData();
+			
+		for (uint32 i = 0; i != NumberOfElements; ++i)
+		{
+			const float FloatData = ElementWiseCosTensorInitializer(DataType, i, TensorIndex);//Default and only initializer for now.
+			const int32 IntData = (int32)FloatData;
+			const uint32 UIntData = (uint32)FloatData;
+			const char BoolData = FloatData != 0.0f ? 1 : 0;
+			switch (DataType)
 			{
-				const float FloatData = Initializer(DataType, i, Index);
-				const int32 IntData = (int32)FloatData;
-				const uint32 UIntData = (uint32)FloatData;
-				const char BoolData = FloatData != 0.0f ? 1 : 0;
-				switch (DataType)
-				{
-					case EMLTensorDataType::Float:
-						check(sizeof(FloatData) == ElementByteSize);
-						FMemory::Memcpy(DestinationPtr, &FloatData, ElementByteSize); break;
-					case EMLTensorDataType::Int32:
-						check(sizeof(IntData) == ElementByteSize);
-						FMemory::Memcpy(DestinationPtr, &IntData, ElementByteSize); break;
-					case EMLTensorDataType::UInt32:
-						check(sizeof(UIntData) == ElementByteSize);
-						FMemory::Memcpy(DestinationPtr, &UIntData, ElementByteSize); break;
-					case EMLTensorDataType::Boolean:
-						check(sizeof(BoolData) == ElementByteSize);
-						FMemory::Memcpy(DestinationPtr, &BoolData, ElementByteSize); break;
-					default:
-						//TODO handle more TensorData types
-						FMemory::Memzero(DestinationPtr, ElementByteSize);
-				}
-				DestinationPtr += ElementByteSize;
+			case EMLTensorDataType::Float:
+				check(sizeof(FloatData) == ElementByteSize);
+				FMemory::Memcpy(DestinationPtr, &FloatData, ElementByteSize); break;
+			case EMLTensorDataType::Int32:
+				check(sizeof(IntData) == ElementByteSize);
+				FMemory::Memcpy(DestinationPtr, &IntData, ElementByteSize); break;
+			case EMLTensorDataType::UInt32:
+				check(sizeof(UIntData) == ElementByteSize);
+				FMemory::Memcpy(DestinationPtr, &UIntData, ElementByteSize); break;
+			case EMLTensorDataType::Boolean:
+				check(sizeof(BoolData) == ElementByteSize);
+				FMemory::Memcpy(DestinationPtr, &BoolData, ElementByteSize); break;
+			default:
+				FMemory::Memzero(DestinationPtr, ElementByteSize);
 			}
+			DestinationPtr += ElementByteSize;
+		}
+
+		return Buffer;
+	}
+
+	static void FillInputTensorBindings(TConstArrayView<FTensor> Tensors, TConstArrayView<FTests::FTensorData> TensorsData, 
+		TArray<TArray<char>>& OutMemBuffers, TArray<FMLTensorBinding>& OutBindings)
+	{
+		OutBindings.Empty();
+		OutMemBuffers.Empty();
+		check(TensorsData.Num() == 0 || Tensors.Num() == TensorsData.Num());
 
-			OutMemBuffers.Emplace(CurInputBuffer);
-			OutBindings.Emplace(
-				FMLTensorBinding::FromCPU(OutMemBuffers[Index].GetData(), BufferSize)
-			);
+		for (int Index = 0; Index < Tensors.Num(); ++Index)
+		{
+			const FTensor& Tensor = Tensors[Index];
+			TArray<char>& MemBuffer = OutMemBuffers.Emplace_GetRef();
+			if (TensorsData.Num() == 0 || TensorsData[Index].Num() == 0)
+			{
+				MemBuffer = GenerateTensorDataForTest(Tensor, Index);
+			}
+			else
+			{
+				MemBuffer = TensorsData[Index];
+			}
+			check(MemBuffer.Num() == Tensor.GetDataSize());
+			OutBindings.Emplace(FMLTensorBinding::FromCPU((void*)(MemBuffer.GetData()), MemBuffer.Num()));
 		}
+	}
 
-		check(TensorDescs.Num() == OutMemBuffers.Num());
-		check(TensorDescs.Num() == OutBindings.Num());
+	static void FillOutputTensorBindings(TConstArrayView<FTensor> Tensors, TArray<TArray<char>>& OutMemBuffers, TArray<FMLTensorBinding>& OutBindings)
+	{
+		OutBindings.Empty();
+		OutMemBuffers.Empty();
+		for (int Index = 0; Index < Tensors.Num(); ++Index)
+		{
+			TArray<char>& MemBuffer = OutMemBuffers.Emplace_GetRef();
+			MemBuffer.SetNumUninitialized(Tensors[Index].GetDataSize());
+			char constexpr MagicNumber = 0x5b;
+			FMemory::Memset(MemBuffer.GetData(), MagicNumber, MemBuffer.Num());
+			OutBindings.Emplace(FMLTensorBinding::FromCPU((void*)(MemBuffer.GetData()), MemBuffer.Num()));
+		}
 	}
 
-	template<typename T> FString ShapeToString(TArrayView<const T> Shape)
+	template<typename T> FString ShapeToString(TConstArrayView<T> Shape)
 	{
 		FString TestSuffix(TEXT("["));
 		bool bIsFirstDim = true;
@@ -97,10 +150,10 @@ namespace Test
 		TestSuffix += TEXT("]");
 		return TestSuffix;
 	}
-	template FString ShapeToString<int32>(TArrayView<const int32> Shape);
-	template FString ShapeToString<uint32>(TArrayView<const uint32> Shape);
+	template FString ShapeToString<int32>(TConstArrayView<int32> Shape);
+	template FString ShapeToString<uint32>(TConstArrayView<uint32> Shape);
 
-	FString FMLTensorDescToString(const FTensor& desc)
+	FString TensorToString(const FTensor& desc)
 	{
 		TArrayView<const uint32> Shape(desc.GetShape().Data);
 
@@ -114,12 +167,12 @@ namespace Test
 		return TensorDesc;
 	}
 
-	FString TensorToString(const FTensor& TensorDesc, const TArray<char>& TensorData)
+	FString TensorToString(const FTensor& TensorDesc, TConstArrayView<char> TensorData)
 	{
 		FString TensorLog;
 		TensorLog.Reserve(50);
 
-		TensorLog += FMLTensorDescToString(TensorDesc);
+		TensorLog += TensorToString(TensorDesc);
 		TensorLog += TEXT(", Data: ");
 
 		const constexpr uint32 MAX_DATA_TO_LOG = 10;
@@ -257,42 +310,34 @@ namespace Test
 	}
 
 	bool VerifyTensorResult(
-		const FTensor& RefTensorDesc, const TArray<char>& RefRawBuffer,
-		const FTensor& OtherTensorDesc, const TArray<char>& OtherRawBuffer,
+		const FTensor& RefTensor, const TArray<char>& RefRawBuffer,
+		const FTensor& OtherTensor, const TArray<char>& OtherRawBuffer,
 		float AbsoluteErrorEpsilon, float RelativeErrorPercent)
 	{
-		bool bTensorDescMatch = true;
-		bTensorDescMatch &= (RefTensorDesc.GetName() == OtherTensorDesc.GetName());
-		bTensorDescMatch &= (RefTensorDesc.GetShape().Rank() == OtherTensorDesc.GetShape().Rank());
-		bTensorDescMatch &= (RefTensorDesc.GetDataType() == OtherTensorDesc.GetDataType());
-		check(RefTensorDesc.GetShape().Rank() <= FTensorShape::MaxRank);
-		for (int32 i = 0; i < RefTensorDesc.GetShape().Rank(); ++i)
-		{
-			bTensorDescMatch &= (RefTensorDesc.GetShape().Data[i] == OtherTensorDesc.GetShape().Data[i]);
-		}
-
-		if (!bTensorDescMatch)
+		if (RefTensor.GetShape() != OtherTensor.GetShape())
 		{
-			UE_LOG(LogNNX, Error, TEXT("Tensor desc do not match.\nExpected: %s\nGot:      %s"), *FMLTensorDescToString(RefTensorDesc), *FMLTensorDescToString(OtherTensorDesc));
+			TArrayView<const uint32> RefShape(RefTensor.GetShape().Data);
+			TArrayView<const uint32> OtherShape(OtherTensor.GetShape().Data);
+			UE_LOG(LogNNX, Error, TEXT("Tensor shape do not match.\nExpected: %s\nGot:      %s"), *ShapeToString(RefShape), *ShapeToString(OtherShape));
 			return false;
 		}
 
-		if (RefTensorDesc.GetDataType() == EMLTensorDataType::Float)
+		if (RefTensor.GetDataType() == EMLTensorDataType::Float)
 		{
-			return CompareTensorData<float>(RefTensorDesc, RefRawBuffer, OtherTensorDesc, OtherRawBuffer, AbsoluteErrorEpsilon, RelativeErrorPercent);
+			return CompareTensorData<float>(RefTensor, RefRawBuffer, OtherTensor, OtherRawBuffer, AbsoluteErrorEpsilon, RelativeErrorPercent);
 		}
-		else if (RefTensorDesc.GetDataType() == EMLTensorDataType::Boolean)
+		else if (RefTensor.GetDataType() == EMLTensorDataType::Boolean)
 		{
-			check(RefTensorDesc.GetElemByteSize() == 1);
-			return CompareTensorData<bool>(RefTensorDesc, RefRawBuffer, OtherTensorDesc, OtherRawBuffer, AbsoluteErrorEpsilon, RelativeErrorPercent);
+			check(RefTensor.GetElemByteSize() == 1);
+			return CompareTensorData<bool>(RefTensor, RefRawBuffer, OtherTensor, OtherRawBuffer, AbsoluteErrorEpsilon, RelativeErrorPercent);
 		}
-		else if (RefTensorDesc.GetDataType() == EMLTensorDataType::Int32)
+		else if (RefTensor.GetDataType() == EMLTensorDataType::Int32)
 		{
-			return CompareTensorData<int32>(RefTensorDesc, RefRawBuffer, OtherTensorDesc, OtherRawBuffer, AbsoluteErrorEpsilon, RelativeErrorPercent);
+			return CompareTensorData<int32>(RefTensor, RefRawBuffer, OtherTensor, OtherRawBuffer, AbsoluteErrorEpsilon, RelativeErrorPercent);
 		}
-		else if (RefTensorDesc.GetDataType() == EMLTensorDataType::UInt32)
+		else if (RefTensor.GetDataType() == EMLTensorDataType::UInt32)
 		{
-			return CompareTensorData<uint32>(RefTensorDesc, RefRawBuffer, OtherTensorDesc, OtherRawBuffer, AbsoluteErrorEpsilon, RelativeErrorPercent);
+			return CompareTensorData<uint32>(RefTensor, RefRawBuffer, OtherTensor, OtherRawBuffer, AbsoluteErrorEpsilon, RelativeErrorPercent);
 		}
 		else
 		{
@@ -303,48 +348,10 @@ namespace Test
 		return true;
 	}
 
-	static float InputTensorInitializer(EMLTensorDataType DataType, uint32 ElementIndex, uint32 TensorIndex)
-	{
-		const constexpr uint32 IndexOffsetBetweenTensor = 9;
-		switch(DataType)
-		{
-			
-			case EMLTensorDataType::Boolean:
-				return (ElementIndex + IndexOffsetBetweenTensor * TensorIndex) % 2;
-			case EMLTensorDataType::Char:
-			case EMLTensorDataType::Int8:
-			case EMLTensorDataType::Int16:
-			case EMLTensorDataType::Int32:
-			case EMLTensorDataType::Int64:
-				return 10.0f * FMath::Cos((float)(ElementIndex + IndexOffsetBetweenTensor * TensorIndex));
-			case EMLTensorDataType::UInt8:
-			case EMLTensorDataType::UInt16:
-			case EMLTensorDataType::UInt32:
-			case EMLTensorDataType::UInt64:
-				return 10.0f * FMath::Abs(FMath::Cos((float)(ElementIndex + IndexOffsetBetweenTensor * TensorIndex)));
-			default:
-			//case EMLTensorDataType::None:
-			//case EMLTensorDataType::Half:
-			//case EMLTensorDataType::Double:
-			//case EMLTensorDataType::Float:
-			//case EMLTensorDataType::Complex64:
-			//case EMLTensorDataType::Complex128:
-			//case EMLTensorDataType::BFloat16:
-				return FMath::Cos((float)(ElementIndex + IndexOffsetBetweenTensor * TensorIndex));
-		}
-	}
-
-	static float OutputTensorInitializer(EMLTensorDataType DataType, uint32 ElementIndex, uint32 TensorIndex)
-	{
-		const constexpr uint32 IndexOffsetBetweenTensor = 13;
-		return FMath::Sin((float)(ElementIndex + IndexOffsetBetweenTensor * TensorIndex));
-	}
-
-	static int RunTestInference(const FNNIModelRaw& ONNXModelData, const FTests::FTestSetup& TestSetup, 
+	static bool RunTestInference(const FNNIModelRaw& ONNXModelData, const FTests::FTestSetup& TestSetup, 
 		IRuntime* Runtime, TArray<FTensor>& OutOutputTensors, TArray<TArray<char>>& OutOutputMemBuffers)
 	{
 		OutOutputMemBuffers.Empty();
-		OutOutputTensors.Empty();
 
 		FOptimizerOptionsMap Options;
 		FNNIModelRaw RuntimeModelData;
@@ -353,7 +360,7 @@ namespace Test
 		if (!Optimizer || !Optimizer->Optimize(ONNXModelData, RuntimeModelData, Options))
 		{
 			UE_LOG(LogNNX, Error, TEXT("Failed to optimize the model"));
-			return -1;
+			return false;
 		}
 		
 		UMLInferenceModel* UInferenceModel = UMLInferenceModel::CreateFromFormatDesc(RuntimeModelData);
@@ -362,7 +369,7 @@ namespace Test
 		if (!InferenceModel.IsValid())
 		{
 			UE_LOG(LogNNX, Error, TEXT("Could not create Inference model."));
-			return -1;
+			return false;
 		}
 
 		// If test does not ask for specific input/output, fill with model default converting variable dim to 1 if any.
@@ -375,13 +382,9 @@ namespace Test
 		TArray<FMLTensorBinding> InputBindings;
 		TArray<FMLTensorBinding> OutputBindings;
 		TArray<TArray<char>> InputMemBuffers;
-		TArray<TArray<char>> OutputMemBuffers;
 		
-		FillTensorBindings(InputTensors, InputMemBuffers, InputBindings, InputTensorInitializer);
-		FillTensorBindings(OutOutputTensors, OutOutputMemBuffers, OutputBindings, OutputTensorInitializer);
-
-		check(InputTensors.Num() == InputMemBuffers.Num());
-		check(OutOutputTensors.Num() == OutOutputMemBuffers.Num());
+		FillInputTensorBindings(InputTensors, TestSetup.InputsData, InputMemBuffers, InputBindings);
+		FillOutputTensorBindings(OutOutputTensors, OutOutputMemBuffers, OutputBindings);
 
 		// To help for debugging sessions.
 		#ifdef UE_BUILD_DEBUG
@@ -400,7 +403,6 @@ namespace Test
 		}
 		#endif //UE_BUILD_DEBUG
 
-		
 		TArray<FTensorShape> InputShapes;
 		for (const FTensor& Tensor : InputTensors)
 		{
@@ -408,50 +410,72 @@ namespace Test
 		}
 
 		// Setup inputs
-		int ReturnValue = InferenceModel->SetInputTensorShapes(InputShapes);
-		if (ReturnValue != 0)
+		if (0 != InferenceModel->SetInputTensorShapes(InputShapes))
 		{
 			UE_LOG(LogNNX, Error, TEXT("Failed to set input tensor shapes."));
-			return ReturnValue;
+			return false;
 		}
 
 		// Run inference
-		ReturnValue = InferenceModel->Run(InputBindings, OutputBindings);
-		if (ReturnValue != 0)
+		if (0 != InferenceModel->Run(InputBindings, OutputBindings))
 		{
 			UE_LOG(LogNNX, Error, TEXT("Failed to run the model."));
-			return ReturnValue;
+			return false;
+		}
+
+		// Verify output shapes are as expected
+		TConstArrayView<FTensorShape> OutputShapes = InferenceModel->GetOutputTensorShapes();
+		if (OutputShapes.Num() != OutOutputTensors.Num())
+		{
+			UE_LOG(LogNNX, Error, TEXT("Expected %d output tensors, got %d."), OutOutputTensors.Num(), OutputShapes.Num());
+			return false;
+		}
+		for (int i = 0; i < OutOutputTensors.Num(); ++i)
+		{
+			const FTensorShape& RefTensorShape = OutOutputTensors[i].GetShape();
+			const FTensorShape& OtherTensorShape = OutputShapes[i];
+			if (RefTensorShape != OtherTensorShape)
+			{
+				TArrayView<const uint32> RefShape(RefTensorShape.Data);
+				TArrayView<const uint32> OtherShape(OtherTensorShape.Data);
+				UE_LOG(LogNNX, Error, TEXT("Output shape do not match at index %d.\nExpected: %s\nGot:      %s"), i, *ShapeToString(RefShape), *ShapeToString(OtherShape));
+				return false;
+			}
 		}
 		
-		return ReturnValue;
+		return true;
 	}
 
 	bool RunTestInferenceAndCompareToRef(const FTests::FTestSetup& TestSetup, IRuntime* Runtime, const FNNIModelRaw& ONNXModel, 
-		TArrayView<TArray<char>> RefOutputMemBuffers, TArrayView<FTensor> RefOutputTensorDescs)
+		TArrayView<TArray<char>> RefOutputMemBuffers, TArrayView<FTensor> RefOutputTensors)
 	{
 		TArray<TArray<char>> OutputMemBuffers;
-		TArray<FTensor> OutputTensorDescs;
+		TArray<FTensor> OutputTensors;
 
 		const FString& RuntimeName = Runtime->GetRuntimeName();
-
+		const float AbsoluteErrorEpsilon = TestSetup.GetAbsoluteErrorEpsilonForRuntime(RuntimeName);
+		const float RelativeErrorPercent = TestSetup.GetRelativeErrorPercentForRuntime(RuntimeName);
 		bool bTestSuceeded = true;
-		float AbsoluteErrorEpsilon = TestSetup.GetAbsoluteErrorEpsilonForRuntime(RuntimeName);
-		float RelativeErrorPercent = TestSetup.GetRelativeErrorPercentForRuntime(RuntimeName);
+	
+		if (!RunTestInference(ONNXModel, TestSetup, Runtime, OutputTensors, OutputMemBuffers))
+		{
+			UE_LOG(LogNNX, Error, TEXT("Error running inference for engine %s."), *RuntimeName);
+			return false;
+		}
 
-		RunTestInference(ONNXModel, TestSetup, Runtime, OutputTensorDescs, OutputMemBuffers);
-		if (OutputTensorDescs.Num() == RefOutputTensorDescs.Num())
+		if (OutputTensors.Num() == RefOutputTensors.Num())
 		{
-			for (int i = 0; i < OutputTensorDescs.Num(); ++i)
+			for (int i = 0; i < OutputTensors.Num(); ++i)
 			{
 				bTestSuceeded &= VerifyTensorResult(
-					RefOutputTensorDescs[i], RefOutputMemBuffers[i],
-					OutputTensorDescs[i], OutputMemBuffers[i],
+					RefOutputTensors[i], RefOutputMemBuffers[i],
+					OutputTensors[i], OutputMemBuffers[i],
 					AbsoluteErrorEpsilon, RelativeErrorPercent);
 			}
 		}
 		else
 		{
-			UE_LOG(LogNNX, Error, TEXT("Expecting %d output tensor(s), got %d."), RefOutputTensorDescs.Num(), OutputTensorDescs.Num());
+			UE_LOG(LogNNX, Error, TEXT("Expecting %d output tensor(s), got %d."), RefOutputTensors.Num(), OutputTensors.Num());
 			bTestSuceeded = false;
 		}
 		return bTestSuceeded;
@@ -475,14 +499,36 @@ namespace Test
 			return false;
 		}
 		const FString& RefName = RefRuntime->GetRuntimeName();
+		const float AbsoluteRefErrorEpsilon = TestSetup.GetAbsoluteErrorEpsilonForRuntime(RefName);
+		const float RelativeRefErrorPercent = TestSetup.GetRelativeErrorPercentForRuntime(RefName);
 		TArray<TArray<char>> RefOutputMemBuffers;
-		TArray<FTensor> RefOutputTensorDescs;
+		TArray<FTensor> RefOutputTensors;
+		bool bAllTestsSucceeded = true;
+
+		if (!RunTestInference(ONNXModel, TestSetup, RefRuntime, RefOutputTensors, RefOutputMemBuffers))
+		{
+			UE_LOG(LogNNX, Error, TEXT("Error running reference inference with engine %s."), *RefName);
+			return false;
+		}
 
-		RunTestInference(ONNXModel, TestSetup, RefRuntime, RefOutputTensorDescs, RefOutputMemBuffers);
+		// If outputs data have been defined by test setup we check that reference match it
+		for (int i = 0; i < TestSetup.OutputsData.Num(); ++i)
+		{
+			if (!TestSetup.OutputsData[i].IsEmpty())
+			{
+				bAllTestsSucceeded &= VerifyTensorResult(
+					TestSetup.Outputs[i], TestSetup.OutputsData[i],
+					RefOutputTensors[i], RefOutputMemBuffers[i],
+					AbsoluteRefErrorEpsilon, RelativeRefErrorPercent);
+			}
+		}
+		if (!bAllTestsSucceeded)
+		{
+			UE_LOG(LogNNX, Error, TEXT("Expecting output from test setup are no match by reference engine."));
+		}
 
 		// Test against other runtime
 		TArray<IRuntime*> Runtimes;
-		bool bAllTestsSucceeded = true;
 
 		Runtimes = GetAllRuntimes();
 		for (auto Runtime : Runtimes)
@@ -516,7 +562,7 @@ namespace Test
 				bool bShouldRunVariadicTest = (ONNXModelVariadic.Format != ENNXInferenceFormat::Invalid);
 				bShouldRunVariadicTest &= !(RuntimeName == "NNXRuntimeDML");
 
-				bool bTestSuceeded = RunTestInferenceAndCompareToRef(TestSetup, Runtime, ONNXModel, RefOutputMemBuffers, RefOutputTensorDescs);
+				bool bTestSuceeded = RunTestInferenceAndCompareToRef(TestSetup, Runtime, ONNXModel, RefOutputMemBuffers, RefOutputTensors);
 
 				if (bShouldRunVariadicTest)
 				{
@@ -524,7 +570,7 @@ namespace Test
 					{
 						UE_LOG(LogNNX, Error, TEXT("Failed running static test."));
 					}
-					bool bVariadicTestSuceeded = RunTestInferenceAndCompareToRef(TestSetup, Runtime, ONNXModelVariadic, RefOutputMemBuffers, RefOutputTensorDescs);
+					bool bVariadicTestSuceeded = RunTestInferenceAndCompareToRef(TestSetup, Runtime, ONNXModelVariadic, RefOutputMemBuffers, RefOutputTensors);
 					if (!bVariadicTestSuceeded)
 					{
 						bTestSuceeded = false;
diff --git a/Engine/Plugins/Experimental/NNX/Source/NNXQA/Private/NNXQAUtils.h b/Engine/Plugins/Experimental/NNX/Source/NNXQA/Private/NNXQAUtils.h
--- a/Engine/Plugins/Experimental/NNX/Source/NNXQA/Private/NNXQAUtils.h
+++ b/Engine/Plugins/Experimental/NNX/Source/NNXQA/Private/NNXQAUtils.h
@@ -14,6 +14,8 @@ namespace Test
 {
 	struct FTests
 	{
+		typedef TArray<char> FTensorData;
+		
 		struct FTestSetup
 		{
 			static constexpr float DefaultAbsoluteErrorEpsilon = 1e-5f;
@@ -35,7 +37,11 @@ namespace Test
 			TMap<FString, float> AbsoluteErrorEpsilonForRuntime;
 			TMap<FString, float> RelativeErrorPercentForRuntime;
 			TArray<FTensor> Inputs;
+			TArray<FTensor> Weights;
 			TArray<FTensor> Outputs;
+			TArray<FTensorData> InputsData;
+			TArray<FTensorData> WeightsData;
+			TArray<FTensorData> OutputsData;
 			UE::NNECore::FAttributeMap AttributeMap;
 			TArray<FString> Tags;
 			TArray<FString> AutomationExcludedRuntime;
@@ -63,9 +69,10 @@ namespace Test
 
 	bool CompareONNXModelInferenceAcrossRuntimes(const FNNIModelRaw& ONNXModel, const FNNIModelRaw& ONNXModelVariadic, const FTests::FTestSetup& TestSetup, const FString& RuntimeFilter = TEXT(""));
 
-	FString FMLTensorDescToString(const FTensor& TensorDesc);
-	FString TensorToString(const FTensor& TensorDesc, const TArray<char>& TensorData);
-	template<typename T> FString ShapeToString(TArrayView<const T> Shape);
+	FString TensorToString(const FTensor& Tensor);
+	FString TensorToString(const FTensor& Tensor, TConstArrayView<char> TensorData);
+	template<typename T> FString ShapeToString(TConstArrayView<T> Shape);
+	TArray<char> GenerateTensorDataForTest(const FTensor& Tensor, int TensorIndex);
 
 } // namespace Test
 } // namespace NNX
\ No newline at end of file
diff --git a/Engine/Plugins/Experimental/NNX/Source/NNXQA/Resources/NNXQATestDescSchema.json b/Engine/Plugins/Experimental/NNX/Source/NNXQA/Resources/NNXQATestDescSchema.json
--- a/Engine/Plugins/Experimental/NNX/Source/NNXQA/Resources/NNXQATestDescSchema.json
+++ b/Engine/Plugins/Experimental/NNX/Source/NNXQA/Resources/NNXQATestDescSchema.json
@@ -37,7 +37,11 @@
 					"type": "array",
 					"items": { "type": "number" }
 			    },
-				"data_type": {
+				"source": {
+					"type": "array",
+					"items": { "type": "string" }
+			    },
+				"type": {
 					"type": "string",
 					"enum": [ "float", "int32", "uint32", "boolean" ]
 				},
@@ -86,12 +90,17 @@
 			"properties": {
 				"inputs": {
 					"type": "array",
-					"description": "input shapes to test",
+					"description": "input to test",
+					"items": { "$ref": "#/$defs/test_config_tensor" }
+				},
+				"weights": {
+					"type": "array",
+					"description": "output to test",
 					"items": { "$ref": "#/$defs/test_config_tensor" }
 				},
 				"outputs": {
 					"type": "array",
-					"description": "output shapes to test",
+					"description": "output to test",
 					"items": { "$ref": "#/$defs/test_config_tensor" }
 				},
 				"runtimes": {
