NNI Refactor test framework to allow test definition to provide the input/output data

Problem:
* Test description need to be able to define the input/output data for to cover more case or do it more explicitely

Solution:
* Add an element wise tensor initializer from json to input and output tensor. Add a related test case for Pad operator.

Tests:
* All NNX tests pass under win64 editor

diff --git a/Engine/Plugins/Experimental/NNX/Source/NNXQA/Private/NNXQAParametricTest.cpp b/Engine/Plugins/Experimental/NNX/Source/NNXQA/Private/NNXQAParametricTest.cpp
--- a/Engine/Plugins/Experimental/NNX/Source/NNXQA/Private/NNXQAParametricTest.cpp
+++ b/Engine/Plugins/Experimental/NNX/Source/NNXQA/Private/NNXQAParametricTest.cpp
@@ -133,15 +133,54 @@ namespace Test
 			return (Value == INDEX_NONE) ? DefaultValue : (EMLTensorDataType)Value;
 		}
 
-		static TArray<char> GetInputTensorDataFromJson(FTensor Tensor, int TensorIndex, TConstArrayView<FString> JsonValues)
+		class ElementWiseFromJsonStringTensorInitializer
 		{
-			checkf(JsonValues.Num() == 0, TEXT("source data for tensor not supported yet"));
-			return GenerateTensorDataForTest(Tensor, TensorIndex);
+			TConstArrayView<FString> JsonValues;
+				
+		public:
+			ElementWiseFromJsonStringTensorInitializer(TConstArrayView<FString> InJsonValues)
+				:JsonValues(InJsonValues) 
+			{}
+
+			float operator () (uint32 TensorIndex) const 
+			{
+				const FString& JsonValue = JsonValues[TensorIndex];
+				if (!FCString::IsNumeric(*JsonValue))
+				{
+					UE_LOG(LogNNX, Error, TEXT("Cannot convert %s to float will default to 0.0f, check test config."), *JsonValue);
+				}
+				return FCString::Atof(*JsonValues[TensorIndex]);
+			}
+		};
+
+		static TArray<char> GetInputTensorDataFromJson(const FTensor& Tensor, int TensorIndex, TConstArrayView<FString> JsonValues)
+		{
+			if (JsonValues.Num() == Tensor.GetVolume())
+			{
+				ElementWiseFromJsonStringTensorInitializer Initializer(JsonValues);
+				return GenerateTensorDataForTest(Tensor, Initializer);
+			}
+			else if (JsonValues.Num() > 0)
+			{
+				UE_LOG(LogNNX, Error, TEXT("Incorrect number of element for tensor initializer %s, should have %d but got %d. Fallbacking to default initializer."), *Tensor.GetName(), Tensor.GetVolume(), JsonValues.Num());
+				
+			}
+			ElementWiseCosTensorInitializer Initializer(Tensor.GetDataType(), TensorIndex);
+			return GenerateTensorDataForTest(Tensor, Initializer);
 		}
 
-		static TArray<char> GetOutputTensorDataFromJson(FTensor Tensor, int TensorIndex, TConstArrayView<FString> JsonValues)
+		static TArray<char> GetOutputTensorDataFromJson(const FTensor& Tensor, int TensorIndex, TConstArrayView<FString> JsonValues)
 		{
-			checkf(JsonValues.Num() == 0, TEXT("source data for tensor not supported yet"));
+			if (JsonValues.Num() == Tensor.GetVolume())
+			{
+				ElementWiseFromJsonStringTensorInitializer Initializer(JsonValues);
+				return GenerateTensorDataForTest(Tensor, Initializer);
+			}
+			else if (JsonValues.Num() > 0)
+			{
+				UE_LOG(LogNNX, Error, TEXT("Incorrect number of element for tensor initializer %s, should have %d but got %d."), *Tensor.GetName());
+
+			}
 			return TArray<char>();
 		}
 
diff --git a/Engine/Plugins/Experimental/NNX/Source/NNXQA/Private/NNXQAUtils.cpp b/Engine/Plugins/Experimental/NNX/Source/NNXQA/Private/NNXQAUtils.cpp
--- a/Engine/Plugins/Experimental/NNX/Source/NNXQA/Private/NNXQAUtils.cpp
+++ b/Engine/Plugins/Experimental/NNX/Source/NNXQA/Private/NNXQAUtils.cpp
@@ -27,8 +27,13 @@ namespace Test
 			}
 		}
 	}
+	
+	ElementWiseCosTensorInitializer::ElementWiseCosTensorInitializer(EMLTensorDataType InDataType, uint32 InTensorIndex)
+	: DataType(InDataType), TensorIndex(InTensorIndex)
+	{
+	}
 
-	static float ElementWiseCosTensorInitializer(EMLTensorDataType DataType, uint32 ElementIndex, uint32 TensorIndex)
+	float ElementWiseCosTensorInitializer::operator() (uint32 ElementIndex) const 
 	{
 		const constexpr uint32 IndexOffsetBetweenTensor = 9;
 		switch (DataType)
@@ -57,9 +62,9 @@ namespace Test
 			//case EMLTensorDataType::BFloat16:
 			return FMath::Cos((float)(ElementIndex + IndexOffsetBetweenTensor * TensorIndex));
 		}
-	}
+	};
 
-	TArray<char> GenerateTensorDataForTest(const FTensor& Tensor, int TensorIndex)
+	TArray<char> GenerateTensorDataForTest(const FTensor& Tensor, std::function<float(uint32)> ElementInitializer)
 	{
 		const uint32 NumberOfElements = Tensor.GetVolume();
 		const int32 ElementByteSize = Tensor.GetElemByteSize();
@@ -69,11 +74,12 @@ namespace Test
 			
 		Buffer.SetNum(BufferSize);
 		char* DestinationPtr = Buffer.GetData();
-			
+
 		for (uint32 i = 0; i != NumberOfElements; ++i)
 		{
-			const float FloatData = ElementWiseCosTensorInitializer(DataType, i, TensorIndex);//Default and only initializer for now.
+			const float FloatData = ElementInitializer(i);
 			const int32 IntData = (int32)FloatData;
+			const int64 Int64Data = (int64)FloatData;
 			const uint32 UIntData = (uint32)FloatData;
 			const char BoolData = FloatData != 0.0f ? 1 : 0;
 			switch (DataType)
@@ -84,6 +90,9 @@ namespace Test
 			case EMLTensorDataType::Int32:
 				check(sizeof(IntData) == ElementByteSize);
 				FMemory::Memcpy(DestinationPtr, &IntData, ElementByteSize); break;
+			case EMLTensorDataType::Int64:
+				check(sizeof(Int64Data) == ElementByteSize);
+				FMemory::Memcpy(DestinationPtr, &Int64Data, ElementByteSize); break;
 			case EMLTensorDataType::UInt32:
 				check(sizeof(UIntData) == ElementByteSize);
 				FMemory::Memcpy(DestinationPtr, &UIntData, ElementByteSize); break;
@@ -112,7 +121,8 @@ namespace Test
 			TArray<char>& MemBuffer = OutMemBuffers.Emplace_GetRef();
 			if (TensorsData.Num() == 0 || TensorsData[Index].Num() == 0)
 			{
-				MemBuffer = GenerateTensorDataForTest(Tensor, Index);
+				ElementWiseCosTensorInitializer Initializer(Tensor.GetDataType(), Index);
+				MemBuffer = GenerateTensorDataForTest(Tensor, Initializer);
 			}
 			else
 			{
diff --git a/Engine/Plugins/Experimental/NNX/Source/NNXQA/Private/NNXQAUtils.h b/Engine/Plugins/Experimental/NNX/Source/NNXQA/Private/NNXQAUtils.h
--- a/Engine/Plugins/Experimental/NNX/Source/NNXQA/Private/NNXQAUtils.h
+++ b/Engine/Plugins/Experimental/NNX/Source/NNXQA/Private/NNXQAUtils.h
@@ -3,10 +3,12 @@
 #pragma once
 
 #include "Containers/UnrealString.h"
-#include "NNXRuntime.h"
+#include <functional>
 #include "NNECoreAttributeMap.h"
+#include "NNXRuntime.h"
+#include "NNXRuntimeFormat.h"
+#include "NNXTypes.h"
 
-struct FNNIModelRaw;
 
 namespace NNX 
 {
@@ -67,12 +69,23 @@ namespace Test
 		TArray<FTestSetup> TestSetups;
 	};
 
-	bool CompareONNXModelInferenceAcrossRuntimes(const FNNIModelRaw& ONNXModel, const FNNIModelRaw& ONNXModelVariadic, const FTests::FTestSetup& TestSetup, const FString& RuntimeFilter = TEXT(""));
+	bool CompareONNXModelInferenceAcrossRuntimes(const FNNIModelRaw& ONNXModel, const FNNIModelRaw& ONNXModelVariadic, 
+		const FTests::FTestSetup& TestSetup, const FString& RuntimeFilter = TEXT(""));
+	
+	class ElementWiseCosTensorInitializer
+	{
+		EMLTensorDataType DataType;
+		uint32 TensorIndex;
+
+	public:
+		ElementWiseCosTensorInitializer(EMLTensorDataType InDataType, uint32 InTensorIndex);
+		float operator () (uint32 ElementIndex) const;
+	};
 
 	FString TensorToString(const FTensor& Tensor);
 	FString TensorToString(const FTensor& Tensor, TConstArrayView<char> TensorData);
 	template<typename T> FString ShapeToString(TConstArrayView<T> Shape);
-	TArray<char> GenerateTensorDataForTest(const FTensor& Tensor, int TensorIndex);
+	TArray<char> GenerateTensorDataForTest(const FTensor& Tensor, std::function<float(uint32)> ElementInitializer);
 
 } // namespace Test
 } // namespace NNX
\ No newline at end of file
diff --git a/Engine/Plugins/Experimental/NNX/Source/NNXQA/Resources/NNXQATestDesc.json b/Engine/Plugins/Experimental/NNX/Source/NNXQA/Resources/NNXQATestDesc.json
--- a/Engine/Plugins/Experimental/NNX/Source/NNXQA/Resources/NNXQATestDesc.json
+++ b/Engine/Plugins/Experimental/NNX/Source/NNXQA/Resources/NNXQATestDesc.json
@@ -292,6 +292,16 @@
 			"targets": [
 				{ "target": "MatMul" }
 			]
+		},
+		{
+			"category": "PadOp",
+			"runtimes": [
+				{ "name": "NNXRuntimeDml", "skip": true	},
+				{ "name": "NNXRuntimeHlsl", "skip": true	}
+			],
+			"targets": [
+				{ "target": "Pad" }
+			]
 		}
 	],
 
@@ -557,6 +567,14 @@
 				{ "inputs": [ { "shape": [ 5, 4 ] }, { "shape": [ 2, 1, 3, 4, 3 ] } ], "outputs": [ { "shape": [ 2, 1, 3, 5, 3 ] } ]},
 				{ "inputs": [ { "shape": [ 2, 1, 3, 5, 4 ] }, { "shape": [ 4, 3 ] } ], "outputs": [ { "shape": [ 2, 1, 3, 5, 3 ] } ]}
 			]
+		},
+		{
+			"name": "PadOp.Simple",
+			"datasets": [
+				{ "inputs":  [ { "shape": [ 1 ], "source": [ "1" ] } , 
+							   { "shape": [ 2 ], "source": [ "1", "2" ] , "type": "int64" } ],
+				  "outputs": [ { "shape": [ 4 ], "source": [ "0", "1", "0", "0" ] } ] }
+			]
 		}
 	],
 
diff --git a/Engine/Plugins/Experimental/NNX/Source/NNXQA/Resources/NNXQATestDescSchema.json b/Engine/Plugins/Experimental/NNX/Source/NNXQA/Resources/NNXQATestDescSchema.json
--- a/Engine/Plugins/Experimental/NNX/Source/NNXQA/Resources/NNXQATestDescSchema.json
+++ b/Engine/Plugins/Experimental/NNX/Source/NNXQA/Resources/NNXQATestDescSchema.json
@@ -43,7 +43,7 @@
 			    },
 				"type": {
 					"type": "string",
-					"enum": [ "float", "int32", "uint32", "boolean" ]
+					"enum": [ "float", "int32", "int64", "uint32", "boolean" ]
 				},
 				"comment": {
 					"type": "string",
@@ -189,11 +189,11 @@
 				},
 				"input_type": {
 					"type": "string",
-					"enum": [ "float", "int32", "uint32", "boolean" ]
+					"enum": [ "float", "int32", "int64", "uint32", "boolean" ]
 				},
 				"output_type": {
 					"type": "string",
-					"enum": [ "float", "int32", "uint32", "boolean" ]
+					"enum": [ "float", "int32", "int64", "uint32", "boolean" ]
 				},
 				"runtimes": {
 					"type": "array",
