NNI Refactor test framework to allow test definition to provide the weight data

Problem:
* Test description need to be able to define the data for model weight, so we can test operator in the context of constant inputs for RDG

Solution:
* Initialize onnx model using data from json config

Tests:
* All NNX tests pass under win64 editor, added an operator test using weight as 2nd input

diff --git a/Engine/Plugins/Experimental/NNX/Source/NNXQA/Private/NNXQAParametricTest.cpp b/Engine/Plugins/Experimental/NNX/Source/NNXQA/Private/NNXQAParametricTest.cpp
--- a/Engine/Plugins/Experimental/NNX/Source/NNXQA/Private/NNXQAParametricTest.cpp
+++ b/Engine/Plugins/Experimental/NNX/Source/NNXQA/Private/NNXQAParametricTest.cpp
@@ -252,7 +252,7 @@ namespace Test
 
 		static FString GetTestSuffix(const Json::FTestConfigDataset& Dataset)
 		{
-			//Build TestSuffix "<inputshape0>_...=><outputshape0>_..."
+			//Build TestSuffix "<inputshape0>_..._w<weightshape0>_w...=><outputshape0>_..."
 			bool bIsFirstShape = true;
 			FString TestSuffix;
 			for (const Json::FTestConfigTensor& Input : Dataset.Inputs)
@@ -261,6 +261,11 @@ namespace Test
 				TestSuffix += ShapeToString<int32>(Input.Shape);
 				bIsFirstShape = false;
 			}
+			for (const Json::FTestConfigTensor& Weight : Dataset.Weights)
+			{
+				TestSuffix += TEXT("_w");
+				TestSuffix += ShapeToString<int32>(Weight.Shape);
+			}
 			TestSuffix += TEXT("=>");
 			//If output is not defined it is the first input shape.
 			if (Dataset.Outputs.Num() == 0)
@@ -269,10 +274,12 @@ namespace Test
 			}
 			else
 			{
+				bIsFirstShape = true;
 				for (const Json::FTestConfigTensor& Output : Dataset.Outputs)
 				{
-					TestSuffix += TEXT("_");
+					if (!bIsFirstShape) TestSuffix += TEXT("_");
 					TestSuffix += ShapeToString<int32>(Output.Shape);
+					bIsFirstShape = false;
 				}
 			}
 			return TestSuffix;
@@ -493,16 +500,19 @@ namespace Test
 		else
 		{
 			// Operator test, create model in memory
-			if (!CreateONNXModelForOperator(false, TestSetup.TargetName, TestSetup.Inputs, TestSetup.Outputs, TestSetup.AttributeMap, ONNXModel))
+			if (!CreateONNXModelForOperator(false, TestSetup.TargetName, TestSetup.Inputs, TestSetup.Outputs, TestSetup.Weights, TestSetup.WeightsData,TestSetup.AttributeMap, ONNXModel))
 			{
 				UE_LOG(LogNNX, Error, TEXT("Failed to create static model for test '%s'. Test ABORTED!"), *TestSetup.TargetName);
 				return false;
 			}
-			if (!CreateONNXModelForOperator(true, TestSetup.TargetName, TestSetup.Inputs, TestSetup.Outputs, TestSetup.AttributeMap, ONNXModelVariadic))
+			if (!CreateONNXModelForOperator(true, TestSetup.TargetName, TestSetup.Inputs, TestSetup.Outputs, TestSetup.Weights, TestSetup.WeightsData, TestSetup.AttributeMap, ONNXModelVariadic))
 			{
 				UE_LOG(LogNNX, Error, TEXT("Failed to create variadic model for test '%s'. Test ABORTED!"), *TestSetup.TargetName);
 				return false;
 			}
+
+            //FFileHelper::SaveArrayToFile(ONNXModel.Data, TEXT("D:\\Opmodel.onnx"));
+            //FFileHelper::SaveArrayToFile(ONNXModelVariadic.Data, TEXT("D:\\OpmodelVar.onnx"));
 		}
 
 		return CompareONNXModelInferenceAcrossRuntimes(ONNXModel, ONNXModelVariadic, TestSetup, RuntimeFilter);
diff --git a/Engine/Plugins/Experimental/NNX/Source/NNXQA/Resources/NNXQATestDesc.json b/Engine/Plugins/Experimental/NNX/Source/NNXQA/Resources/NNXQATestDesc.json
--- a/Engine/Plugins/Experimental/NNX/Source/NNXQA/Resources/NNXQATestDesc.json
+++ b/Engine/Plugins/Experimental/NNX/Source/NNXQA/Resources/NNXQATestDesc.json
@@ -569,12 +569,20 @@
 			]
 		},
 		{
-			"name": "PadOp.Simple",
+			"name": "PadOp.PaddingFromInput",
 			"datasets": [
 				{ "inputs":  [ { "shape": [ 1 ], "source": [ "1" ] } , 
 							   { "shape": [ 2 ], "source": [ "1", "2" ] , "type": "int64" } ],
 				  "outputs": [ { "shape": [ 4 ], "source": [ "0", "1", "0", "0" ] } ] }
 			]
+		},
+		{
+			"name": "PadOp.PaddingFromWeight",
+			"datasets": [
+				{ "inputs":  [ { "shape": [ 1 ], "source": [ "1" ] } ], 
+				  "weights": [ { "shape": [ 2 ], "source": [ "1", "2" ] , "type": "int64" } ],
+				  "outputs": [ { "shape": [ 4 ], "source": [ "0", "1", "0", "0" ] } ] }
+			]
 		}
 	],
 
diff --git a/Engine/Plugins/Experimental/NNX/Source/NNXUtils/Private/NNXModelBuilderONNX.cpp b/Engine/Plugins/Experimental/NNX/Source/NNXUtils/Private/NNXModelBuilderONNX.cpp
--- a/Engine/Plugins/Experimental/NNX/Source/NNXUtils/Private/NNXModelBuilderONNX.cpp
+++ b/Engine/Plugins/Experimental/NNX/Source/NNXUtils/Private/NNXModelBuilderONNX.cpp
@@ -123,7 +123,19 @@ public:
 
 		if (Data)
 		{
-			UE_LOG(LogNNX, Warning, TEXT("ONNX ModelBuilder currently doesn't support tensors with data"));
+			onnx::TensorProto* Tensor = Graph->mutable_initializer()->Add();
+
+			Tensor->set_name(TCHAR_TO_ANSI(*Name));
+			Tensor->set_data_type(ToTensorProtoDataType(DataType));
+			for (int32 Idx = 0; Idx < Shape.Num(); ++Idx)
+			{
+				Tensor->add_dims(Shape[Idx]);
+			}
+
+			checkf(DataType != EMLTensorDataType::Char, TEXT("Char tensors not supported yet!"));
+			// raw_data is not supported for strings, when implementing those one will need to use string_data from TensorProto
+			std::string* raw_data = Tensor->mutable_raw_data();
+			raw_data->assign((const char*)Data, DataSize);
 		}
 
 		return MakeTensorHandle(Value);
@@ -253,7 +265,7 @@ private:
 		onnx::TypeProto*		Type = Value->mutable_type();
 		onnx::TypeProto_Tensor* TensorType = Type->mutable_tensor_type();
 		onnx::TensorShapeProto* Shape = TensorType->mutable_shape();
-
+		
 		Value->set_name(TCHAR_TO_ANSI(*Name));
 		TensorType->set_elem_type(ToTensorProtoDataType(DataType));
 
@@ -293,15 +305,6 @@ private:
 	}
 };
 
-//
-//
-//
-NNXUTILS_API bool CreateONNXModelForOperator(bool UseVariadicShapeForModel, const FString& OperatorName, TConstArrayView<FTensor> InInputTensors, TConstArrayView<FTensor> InOutputTensors, FNNIModelRaw& Model)
-{
-	UE::NNECore::FAttributeMap EmptyAttributeMap;
-	return CreateONNXModelForOperator(UseVariadicShapeForModel, OperatorName, InInputTensors, InOutputTensors, EmptyAttributeMap, Model);
-}
-
 //
 //
 //
@@ -322,7 +325,10 @@ void BuildShapeForModel(bool ConvertToVariadicShape, const FTensorShape& InShape
 //
 //
 //
-NNXUTILS_API bool CreateONNXModelForOperator(bool UseVariadicShapeForModel, const FString& OperatorName, TConstArrayView<FTensor> InInputTensors, TConstArrayView<FTensor> InOutputTensors, const UE::NNECore::FAttributeMap& Attributes, FNNIModelRaw& Model)
+NNXUTILS_API bool CreateONNXModelForOperator(bool UseVariadicShapeForModel, const FString& OperatorName,
+	TConstArrayView<FTensor> InInputTensors, TConstArrayView<FTensor> InOutputTensors,
+	TConstArrayView<FTensor> InWeightTensors, TConstArrayView<TArray<char>> InWeightTensorsData,
+	const UE::NNECore::FAttributeMap& Attributes, FNNIModelRaw& Model)
 {
 	Model = FNNIModelRaw{};
 	
@@ -343,7 +349,7 @@ NNXUTILS_API bool CreateONNXModelForOperator(bool UseVariadicShapeForModel, cons
 		Builder->AddInput(Tensor);
 	}
 
-	TArray<IMLModelBuilder::HTensor>	OutputTensors;
+	TArray<IMLModelBuilder::HTensor> OutputTensors;
 
 	for (int32 Idx = 0; Idx < InOutputTensors.Num(); ++Idx)
 	{
@@ -355,6 +361,20 @@ NNXUTILS_API bool CreateONNXModelForOperator(bool UseVariadicShapeForModel, cons
 		Builder->AddOutput(Tensor);
 	}
 
+	checkf(InWeightTensors.Num() == InWeightTensorsData.Num(), TEXT("Invalid weight tensors data"));
+	TArray<IMLModelBuilder::HTensor> WeightTensors;
+
+	for (int32 Idx = 0; Idx < InWeightTensors.Num(); ++Idx)
+	{
+		const FTensor& Desc = InWeightTensors[Idx];
+		const TArray<char>& Data = InWeightTensorsData[Idx];
+		check(Data.Num() == Desc.GetDataSize());
+		BuildShapeForModel(false, Desc.GetShape(), ShapeForModel);
+		IMLModelBuilder::HTensor Tensor = Builder->AddTensor(Desc.GetName(), Desc.GetDataType(), ShapeForModel, Data.GetData(), Data.Num());
+
+		WeightTensors.Emplace(Tensor);
+	}
+
 	auto Op = Builder->AddOperator(OperatorName);
 
 	for (int32 Idx = 0; Idx < InputTensors.Num(); ++Idx)
@@ -362,6 +382,13 @@ NNXUTILS_API bool CreateONNXModelForOperator(bool UseVariadicShapeForModel, cons
 		Builder->AddOperatorInput(Op, InputTensors[Idx]);
 	}
 
+	for (int32 Idx = 0; Idx < WeightTensors.Num(); ++Idx)
+	{
+		// For now weights are added after model inputs in the list of
+		// operator inputs. This should be made more flexible if needed by future tests cases.
+		Builder->AddOperatorInput(Op, WeightTensors[Idx]);
+	}
+
 	for (int32 Idx = 0; Idx < OutputTensors.Num(); ++Idx)
 	{
 		Builder->AddOperatorOutput(Op, OutputTensors[Idx]);
diff --git a/Engine/Plugins/Experimental/NNX/Source/NNXUtils/Private/NNXUtilsModule.cpp b/Engine/Plugins/Experimental/NNX/Source/NNXUtils/Private/NNXUtilsModule.cpp
--- a/Engine/Plugins/Experimental/NNX/Source/NNXUtils/Private/NNXUtilsModule.cpp
+++ b/Engine/Plugins/Experimental/NNX/Source/NNXUtils/Private/NNXUtilsModule.cpp
@@ -37,7 +37,7 @@ using namespace NNX;
 
 static FAutoConsoleCommand ConsoleCmdNNXOptimizeModel(
 	TEXT("nnx.test.OptimizeModel"),
-	TEXT("Optimize the model from file and store it into NNX format.Examples: OptimizeModel -f path.onnx, OptimizeModel -op Cos"),
+	TEXT("Optimize the model from file and store it into NNX format. Example: OptimizeModel -f path.onnx"),
 	FConsoleCommandWithArgsDelegate::CreateStatic(
 		[](const TArray<FString>& Args)
 		{
@@ -60,26 +60,11 @@ static FAutoConsoleCommand ConsoleCmdNNXOptimizeModel(
 						return;
 					}
 				}
-				else if (Args[ArgIdx] == TEXT("-op"))
-				{
-					if (ArgIdx + 1 < Args.Num())
-					{
-						OpName = Args[ArgIdx + 1];
-						break;
-					}
-					else
-					{
-						UE_LOG(LogNNX, Warning, TEXT("Found option -op but operator name is not not provided"));
-						return;
-					}
-				}
-
 			}
 
 			FNNIModelRaw ONNXModel;
 			
-			// First try to load model from file, otherwise create a single layer NN with specified
-			// operator name
+			// First try to load model from file
 			if (!ModelPath.IsEmpty())
 			{
 				if (!FFileHelper::LoadFileToArray(ONNXModel.Data, *ModelPath))
@@ -89,16 +74,6 @@ static FAutoConsoleCommand ConsoleCmdNNXOptimizeModel(
 				}
 				ONNXModel.Format = ENNXInferenceFormat::ONNX;
 			}
-			else if (!OpName.IsEmpty())
-			{
-				FTensor InputTensor = FTensor::Make(TEXT("in"), FTensorShape::Make({1, 512 }), EMLTensorDataType::Float);
-				FTensor OutputTensor = FTensor::Make(TEXT("out"), FTensorShape::Make({ 1, 512 }), EMLTensorDataType::Float);
-
-				if (!NNX::CreateONNXModelForOperator(true, OpName, MakeArrayView(&InputTensor, 1), MakeArrayView(&OutputTensor, 1), ONNXModel))
-				{
-					UE_LOG(LogNNX, Warning, TEXT("Failed to create model for operator:%s"), *OpName);
-				}
-			}
 			else
 			{
 				UE_LOG(LogNNX, Warning, TEXT("Invalid arguments"));
@@ -127,27 +102,4 @@ static FAutoConsoleCommand ConsoleCmdNNXOptimizeModel(
 			}
 		}
 	)
-);
-
-static FAutoConsoleCommand ConsoleCmdNNXCreateModel(
-	TEXT("nnx.test.CreateONNX"),
-	TEXT("Create a single layer ONNX model"),
-	FConsoleCommandWithArgsDelegate::CreateStatic(
-		[](const TArray<FString>& Args)
-		{
-			using NNX::FTensor;
-
-			if (Args.Num() > 0)
-			{
-				FNNIModelRaw ONNXModel;
-				FTensor InputTensor = FTensor::Make(TEXT("in"), FTensorShape::Make({ 1, 512 }), EMLTensorDataType::Float);
-				FTensor OutputTensor = FTensor::Make(TEXT("out"), FTensorShape::Make({ 1, 512 }), EMLTensorDataType::Float);
-
-				if (!NNX::CreateONNXModelForOperator(true, Args[0], MakeArrayView(&InputTensor, 1), MakeArrayView(&OutputTensor, 1), ONNXModel))
-				{
-					UE_LOG(LogNNX, Display, TEXT("Failed to create model for operator:%s"), *Args[0]);
-				}
-			}
-		}
-	)
 );
\ No newline at end of file
diff --git a/Engine/Plugins/Experimental/NNX/Source/NNXUtils/Public/NNXModelBuilder.h b/Engine/Plugins/Experimental/NNX/Source/NNXUtils/Public/NNXModelBuilder.h
--- a/Engine/Plugins/Experimental/NNX/Source/NNXUtils/Public/NNXModelBuilder.h
+++ b/Engine/Plugins/Experimental/NNX/Source/NNXUtils/Public/NNXModelBuilder.h
@@ -87,8 +87,10 @@ NNXUTILS_API IMLModelBuilder* CreateONNXModelBuilder(int64 IrVersion = OnnxIrVer
 /**
  * Utility functions to create single layer NN for operator testing with optional attributes
  */
-NNXUTILS_API bool CreateONNXModelForOperator(bool UseVariadicShapeForModel, const FString& OperatorName, TConstArrayView<FTensor> InInputTensors, TConstArrayView<FTensor> InOutputTensors, FNNIModelRaw& ModelData);
-NNXUTILS_API bool CreateONNXModelForOperator(bool UseVariadicShapeForModel, const FString& OperatorName, TConstArrayView<FTensor> InInputTensors, TConstArrayView<FTensor> InOutputTensors, const UE::NNECore::FAttributeMap& Attributes, FNNIModelRaw& ModelData);
+NNXUTILS_API bool CreateONNXModelForOperator(bool UseVariadicShapeForModel, const FString& OperatorName, 
+	TConstArrayView<FTensor> InInputTensors, TConstArrayView<FTensor> InOutputTensors,
+	TConstArrayView<FTensor> InWeightTensors, TConstArrayView<TArray<char>> InWeightTensorsData,
+	const UE::NNECore::FAttributeMap& Attributes, FNNIModelRaw& ModelData);
 
 /**
  * Create an instance of NNX model builder that creates NNX model/format in memory
