NNE - Adding type and version to model data byte arrays

diff --git a/Engine/Plugins/Experimental/NNX/NNX.uplugin b/Engine/Plugins/Experimental/NNX/NNX.uplugin
--- a/Engine/Plugins/Experimental/NNX/NNX.uplugin
+++ b/Engine/Plugins/Experimental/NNX/NNX.uplugin
@@ -82,8 +82,7 @@
 		{
 			"Name": "NNEEditorTools",
 			"Type": "Editor",
-			"LoadingPhase": "Default",
-			"PlatformAllowList": [ "Win64", "Linux", "Mac" ]
+			"LoadingPhase": "Default"
 		}
 	]
 }
diff --git a/Engine/Plugins/Experimental/NNX/Source/NNXCore/Private/NNEModel.cpp b/Engine/Plugins/Experimental/NNX/Source/NNXCore/Private/NNEModel.cpp
--- a/Engine/Plugins/Experimental/NNX/Source/NNXCore/Private/NNEModel.cpp
+++ b/Engine/Plugins/Experimental/NNX/Source/NNXCore/Private/NNEModel.cpp
@@ -8,6 +8,18 @@
 
 #include "GenericPlatform/GenericPlatformProcess.h"
 
+TArray<FString> UNNEModel::GetRuntimeNames()
+{
+	using namespace NNX;
+	TArray<FString> Result;
+	TArray<NNX::IRuntime*> Runtimes = GetAllRuntimes();
+	for (int i = 0; i < Runtimes.Num(); i++)
+	{
+		Result.Add(Runtimes[i]->GetRuntimeName());
+	}
+	return Result;
+}
+
 UNNEModel* UNNEModel::Create(UObject* Parent, FString RuntimeName, UNNEModelData* ModelData)
 {
 	UNNEModel* Result = NewObject<UNNEModel>(Parent);
@@ -47,16 +59,7 @@ bool UNNEModel::Load(FString RuntimeName, UNNEModelData* ModelData)
 		return false;
 	}
 
-	FNNIModelRaw ModelRaw;
-	ModelRaw.Data = MoveTemp(Data);
-	ModelRaw.Format = ENNXInferenceFormat::ONNX;
-	UMLInferenceModel* TempModelData = UMLInferenceModel::CreateFromFormatDesc(ModelRaw);
-	if (!TempModelData)
-	{
-		return false;
-	}
-
-	Model = TSharedPtr<NNX::FMLInferenceModel>(Runtime->CreateInferenceModel(TempModelData));
+	Model = Runtime->CreateModel(Data);
 
 	return Model.IsValid();
 }
@@ -140,17 +143,21 @@ bool UNNEModel::SetInputOutput(const TArray<FNNETensor>& Input, UPARAM(ref) TArr
 	TArray<NNX::FTensorShape> InputShapes;
 	for (int i = 0; i < Input.Num(); i++)
 	{
-		InputBindings.Add(NNX::FMLTensorBinding::FromCPU((void*)(Input[i].Data.GetData()), 4 * Input[i].Data.Num()));
+		InputBindings.Add(NNX::FMLTensorBinding::FromCPU((void*)(Input[i].Data.GetData()), sizeof(float) * Input[i].Data.Num()));
 
 		InputShapes.Add(NNX::FTensorShape::MakeFromSymbolic(NNX::FSymbolicTensorShape::Make(Input[i].Shape)));
 	}
 
 	for (int i = 0; i < Input.Num(); i++)
 	{
-		OutputBindings.Add(NNX::FMLTensorBinding::FromCPU((void*)(Output[i].Data.GetData()), 4 * Output[i].Data.Num()));
+		OutputBindings.Add(NNX::FMLTensorBinding::FromCPU((void*)(Output[i].Data.GetData()), sizeof(float) * Output[i].Data.Num()));
 	}
 
-	Model->SetInputTensorShapes(InputShapes);
+	if (Model->SetInputTensorShapes(InputShapes) != 0)
+	{
+		UE_LOG(LogNNX, Error, TEXT("UNNEModel: Failed to set input shapes"));
+		return false;
+	}
 
 	return true;
 }
@@ -163,9 +170,4 @@ bool UNNEModel::Run()
 	}
 
 	return Model->Run(InputBindings, OutputBindings) == 0;
-}
-
-TSharedPtr<NNX::FMLInferenceModel> UNNEModel::GetModel()
-{
-	return Model;
 }
\ No newline at end of file
diff --git a/Engine/Plugins/Experimental/NNX/Source/NNXCore/Private/NNEModelData.cpp b/Engine/Plugins/Experimental/NNX/Source/NNXCore/Private/NNEModelData.cpp
--- a/Engine/Plugins/Experimental/NNX/Source/NNXCore/Private/NNEModelData.cpp
+++ b/Engine/Plugins/Experimental/NNX/Source/NNXCore/Private/NNEModelData.cpp
@@ -57,7 +57,7 @@ inline void PutIntoDDC(const FGuid& FileDataId, const FString& RuntimeName, FSha
 
 #endif
 
-inline TArray<uint8> Create(const FString& RuntimeName, const TArray<uint8>& FileData)
+inline TArray<uint8> Create(const FString& RuntimeName, FString FileType, const TArray<uint8>& FileData)
 {
 	using namespace NNX;
 
@@ -74,25 +74,7 @@ inline TArray<uint8> Create(const FString& RuntimeName, const TArray<uint8>& Fil
 		return {};
 	}
 
-	TUniquePtr<IModelOptimizer> Optimizer = Runtime->CreateModelOptimizer();
-	if (!Optimizer.IsValid())
-	{
-		UE_LOG(LogNNX, Error, TEXT("UNNEModelData: %s cannot create an optimizer"), *RuntimeName);
-		return {};
-	}
-
-	FNNIModelRaw InputModel;
-	InputModel.Data = FileData;
-	InputModel.Format = ENNXInferenceFormat::ONNX;
-	FNNIModelRaw OutputModel;
-	FOptimizerOptionsMap Options;
-	if (!Optimizer->Optimize(InputModel, OutputModel, Options))
-	{
-		UE_LOG(LogNNX, Error, TEXT("UNNEModelData: Optimization failed"));
-		return {};
-	}
-
-	return MoveTemp(OutputModel.Data);
+	return Runtime->CreateModelData(FileType, FileData);
 }
 
 void UNNEModelData::Init(const TCHAR* Type, const uint8*& Buffer, const uint8* BufferEnd)
@@ -135,7 +117,7 @@ TConstArrayView<uint8> UNNEModelData::GetModelData(const FString& RuntimeName)
 	}
 
 	// Try to create the model
-	TArray<uint8> CreatedData = Create(RuntimeName, FileData);
+	TArray<uint8> CreatedData = Create(RuntimeName, FileType, FileData);
 	if (CreatedData.Num() < 1)
 	{
 		return {};
@@ -181,7 +163,7 @@ void UNNEModelData::Serialize(FArchive& Ar)
 		TArray<NNX::IRuntime*> Runtimes = GetAllRuntimes();
 		for (int i = 0; i < Runtimes.Num(); i++)
 		{
-			TArray<uint8> CreatedData = Create(Runtimes[i]->GetRuntimeName(), FileData);
+			TArray<uint8> CreatedData = Create(Runtimes[i]->GetRuntimeName(), FileType, FileData);
 			if (CreatedData.Num() > 0)
 			{
 				ModelData.Add(Runtimes[i]->GetRuntimeName(), CreatedData);
diff --git a/Engine/Plugins/Experimental/NNX/Source/NNXCore/Public/NNEModel.h b/Engine/Plugins/Experimental/NNX/Source/NNXCore/Public/NNEModel.h
--- a/Engine/Plugins/Experimental/NNX/Source/NNXCore/Public/NNEModel.h
+++ b/Engine/Plugins/Experimental/NNX/Source/NNXCore/Public/NNEModel.h
@@ -30,6 +30,9 @@ class NNXCORE_API UNNEModel : public UObject
 
 public:
 
+	UFUNCTION(BlueprintCallable, Category = "NNE - Neural Network Engine")
+	static TArray<FString> GetRuntimeNames();
+
 	/**
 	 * Creates a new model.
 	 */
@@ -40,21 +43,21 @@ public:
 	 * Loads the model from model data for a given runtime.
 	 */
 	UFUNCTION(BlueprintCallable, Category = "NNE - Neural Network Engine")
-	virtual bool Load(FString RuntimeName, UNNEModelData* ModelData);
+	bool Load(FString RuntimeName, UNNEModelData* ModelData);
 
 public:
 
 	UFUNCTION(BlueprintCallable, Category = "NNE - Neural Network Engine")
-	virtual int32 NumInputs();
+	int32 NumInputs();
 
 	UFUNCTION(BlueprintCallable, Category = "NNE - Neural Network Engine")
-	virtual int32 NumOutputs();
+	int32 NumOutputs();
 
 	UFUNCTION(BlueprintCallable, Category = "NNE - Neural Network Engine")
-	virtual TArray<int32> GetInputShapes(int32 Index);
+	TArray<int32> GetInputShapes(int32 Index);
 
 	UFUNCTION(BlueprintCallable, Category = "NNE - Neural Network Engine")
-	virtual TArray<int32> GetOutputShapes(int32 Index);
+	TArray<int32> GetOutputShapes(int32 Index);
 
 public:
 
@@ -63,7 +66,7 @@ public:
 	 * If changes are made to the shape or data sizes inside the passed tensors, this function must be called again. 
 	 */
 	UFUNCTION(BlueprintCallable, Category = "NNE - Neural Network Engine")
-	virtual bool SetInputOutput(const TArray<FNNETensor>& Input, UPARAM(ref) TArray<FNNETensor>& Output);
+	bool SetInputOutput(const TArray<FNNETensor>& Input, UPARAM(ref) TArray<FNNETensor>& Output);
 	
 public:
 
@@ -74,14 +77,7 @@ public:
 	 * Upon successful return of this function, the data inside the output tensors passed to SetInputOutput will contain the resulting data.
 	 */
 	UFUNCTION(BlueprintCallable, Category = "NNE - Neural Network Engine")
-	virtual bool Run();
-	
-public:
-	
-	/**
-	 * Get the model to gain full access including running from e.g. the render thread.
-	 */
-	TSharedPtr<NNX::FMLInferenceModel> GetModel();
+	bool Run();
 	
 private:
 
diff --git a/Engine/Plugins/Experimental/NNX/Source/NNXCore/Public/NNXRuntime.h b/Engine/Plugins/Experimental/NNX/Source/NNXCore/Public/NNXRuntime.h
--- a/Engine/Plugins/Experimental/NNX/Source/NNXCore/Public/NNXRuntime.h
+++ b/Engine/Plugins/Experimental/NNX/Source/NNXCore/Public/NNXRuntime.h
@@ -21,7 +21,7 @@ class FRDGBuilder;
 namespace NNX
 {
 class FMLInferenceModel;
-class IModelOptimizer;
+//class IModelOptimizer;
 
 //class IOperatorRegistry;
 
@@ -40,7 +40,9 @@ class IRuntime
 public:
 	virtual ~IRuntime() = default;
 	virtual FString GetRuntimeName() const = 0;
-	virtual TUniquePtr<IModelOptimizer> CreateModelOptimizer() const = 0;
+
+	// Each runtime has already created its own internal optimizer called inside CreateModelData
+	//virtual TUniquePtr<IModelOptimizer> CreateModelOptimizer() const = 0;
 
 	// TODO: Replace with QuerySupport()
 	// Returns flags from ERuntimeSupportFlags
@@ -51,7 +53,15 @@ public:
 
 	// ModelOptimizer
 
-	virtual FMLInferenceModel* CreateInferenceModel(UMLInferenceModel* Model) = 0;
+	//virtual FMLInferenceModel* CreateInferenceModel(UMLInferenceModel* Model) = 0;
+
+	// Model data factory interface
+	virtual bool CanCreateModelData(FString FileType, TConstArrayView<uint8> FileData) const = 0;
+	virtual TArray<uint8> CreateModelData(FString FileType, TConstArrayView<uint8> FileData) = 0;
+
+	// Model factory interface
+	virtual bool CanCreateModel(TConstArrayView<uint8> ModelData) const = 0;
+	virtual TSharedPtr<FMLInferenceModel> CreateModel(TConstArrayView<uint8> ModelData) = 0;
 };
 
 /** Tensor memory binding type */
diff --git a/Engine/Plugins/Experimental/NNX/Source/NNXQA/Private/NNXQAUtils.cpp b/Engine/Plugins/Experimental/NNX/Source/NNXQA/Private/NNXQAUtils.cpp
--- a/Engine/Plugins/Experimental/NNX/Source/NNXQA/Private/NNXQAUtils.cpp
+++ b/Engine/Plugins/Experimental/NNX/Source/NNXQA/Private/NNXQAUtils.cpp
@@ -363,18 +363,7 @@ namespace Test
 	{
 		OutOutputMemBuffers.Empty();
 
-		FOptimizerOptionsMap Options;
-		FNNIModelRaw RuntimeModelData;
-		TUniquePtr<IModelOptimizer> Optimizer = Runtime->CreateModelOptimizer();
-		
-		if (!Optimizer || !Optimizer->Optimize(ONNXModelData, RuntimeModelData, Options))
-		{
-			UE_LOG(LogNNX, Error, TEXT("Failed to optimize the model"));
-			return false;
-		}
-		
-		UMLInferenceModel* UInferenceModel = UMLInferenceModel::CreateFromFormatDesc(RuntimeModelData);
-		TUniquePtr<FMLInferenceModel> InferenceModel(Runtime->CreateInferenceModel(UInferenceModel));
+		TSharedPtr<FMLInferenceModel> InferenceModel = Runtime->CreateModel(Runtime->CreateModelData(FString("onnx"), ONNXModelData.Data));
 		
 		if (!InferenceModel.IsValid())
 		{
diff --git a/Engine/Plugins/Experimental/NNX/Source/NNXRuntimeCPU/Private/NNXRuntimeCPU.cpp b/Engine/Plugins/Experimental/NNX/Source/NNXRuntimeCPU/Private/NNXRuntimeCPU.cpp
--- a/Engine/Plugins/Experimental/NNX/Source/NNXRuntimeCPU/Private/NNXRuntimeCPU.cpp
+++ b/Engine/Plugins/Experimental/NNX/Source/NNXRuntimeCPU/Private/NNXRuntimeCPU.cpp
@@ -2,44 +2,92 @@
 #include "NNXRuntimeCPU.h"
 #include "NNXRuntimeCPUUtils.h"
 #include "NNXModelOptimizer.h"
+#include "NNECoreAttributeMap.h"
 #include "NNEProfilingTimer.h"
 #include "RedirectCoutAndCerrToUeLog.h"
 
 using namespace NNX;
 
+FGuid FRuntimeCPU::GUID = FGuid((int32)'R', (int32)'C', (int32)'P', (int32)'U');
+int32 FRuntimeCPU::Version = 0x00000001;
+
 FString FRuntimeCPU::GetRuntimeName() const
 {
 	return NNX_RUNTIME_CPU_NAME;
 }
 
-TUniquePtr<IModelOptimizer> FRuntimeCPU::CreateModelOptimizer() const
-{
-	return CreateONNXToONNXModelOptimizer();
-}
-
 EMLRuntimeSupportFlags FRuntimeCPU::GetSupportFlags() const 
 {
 	return EMLRuntimeSupportFlags::CPU;
 }
 
+bool FRuntimeCPU::CanCreateModelData(FString FileType, TConstArrayView<uint8> FileData) const
+{
+	return FileType.Compare("onnx", ESearchCase::IgnoreCase) == 0;
+}
 
-FMLInferenceModel* FRuntimeCPU::CreateInferenceModel(UMLInferenceModel* InModel, const FMLInferenceNNXCPUConf& InConf)
+TArray<uint8> FRuntimeCPU::CreateModelData(FString FileType, TConstArrayView<uint8> FileData)
 {
-	FMLInferenceModelCPU* CPUModel = new FMLInferenceModelCPU(&NNXEnvironmentCPU, InConf);
-	if (!CPUModel->Init(InModel))
+	if (!CanCreateModelData(FileType, FileData))
 	{
-		delete CPUModel;
-		CPUModel = nullptr;
+		return {};
 	}
 
-	return CPUModel;
+	TUniquePtr<IModelOptimizer> Optimizer = CreateONNXToONNXModelOptimizer();
+
+	FNNIModelRaw InputModel;
+	InputModel.Data = FileData;
+	InputModel.Format = ENNXInferenceFormat::ONNX;
+	FNNIModelRaw OutputModel;
+	FOptimizerOptionsMap Options;
+	if (!Optimizer->Optimize(InputModel, OutputModel, Options))
+	{
+		return {};
+	}
+
+	int32 GuidSize = sizeof(FRuntimeCPU::GUID);
+	int32 VersionSize = sizeof(FRuntimeCPU::Version);
+	TArray<uint8> Result;
+	FMemoryWriter Writer(Result);
+	Writer << FRuntimeCPU::GUID;
+	Writer << FRuntimeCPU::Version;
+	Writer.Serialize(OutputModel.Data.GetData(), OutputModel.Data.Num());
+	return Result;
 }
 
+bool FRuntimeCPU::CanCreateModel(TConstArrayView<uint8> ModelData) const
+{
+	int32 GuidSize = sizeof(FRuntimeCPU::GUID);
+	int32 VersionSize = sizeof(FRuntimeCPU::Version);
+	if (ModelData.Num() <= GuidSize + VersionSize)
+	{
+		return false;
+	}
+	bool bResult = FGenericPlatformMemory::Memcmp(&(ModelData[0]), &(FRuntimeCPU::GUID), GuidSize) == 0;
+	bResult &= FGenericPlatformMemory::Memcmp(&(ModelData[GuidSize]), &(FRuntimeCPU::Version), VersionSize) == 0;
+	return bResult;
+}
 
-FMLInferenceModel* FRuntimeCPU::CreateInferenceModel(UMLInferenceModel* InModel)
+TSharedPtr<FMLInferenceModel> FRuntimeCPU::CreateModel(TConstArrayView<uint8> ModelData)
 {
-	FMLInferenceNNXCPUConf CPUInferenceConf;
-	return CreateInferenceModel(InModel, CPUInferenceConf);
+	if (!CanCreateModel(ModelData))
+	{
+		return TSharedPtr<FMLInferenceModel>();
+	}
+
+	// Get the header size
+	int32 GuidSize = sizeof(FRuntimeCPU::GUID);
+	int32 VersionSize = sizeof(FRuntimeCPU::Version);
+
+	// Create the model and initialize it with the data not including the header
+	const FMLInferenceNNXCPUConf InConf;
+	FMLInferenceModelCPU* Model = new FMLInferenceModelCPU(&NNXEnvironmentCPU, InConf);
+	if (!Model->Init(ModelData))
+	{
+		delete Model;
+		return TSharedPtr<FMLInferenceModel>();
+	}
+	return TSharedPtr<FMLInferenceModel>(Model);
 }
 
 bool FRuntimeCPU::Init()
@@ -66,13 +114,17 @@ FMLInferenceModelCPU::FMLInferenceModelCPU(
 	ORTEnvironment(InORTEnvironment)
 { }
 
-bool FMLInferenceModelCPU::Init(const UMLInferenceModel* InferenceModel)
+bool FMLInferenceModelCPU::Init(TConstArrayView<uint8> ModelData)
 {
 	DECLARE_SCOPE_CYCLE_COUNTER(TEXT("FMLInferenceModelCPU_Init"), STAT_FMLInferenceModelCPU_Init, STATGROUP_MachineLearning);
 	
+	// Get the header size
+	int32 GuidSize = sizeof(FRuntimeCPU::GUID);
+	int32 VersionSize = sizeof(FRuntimeCPU::Version);
+
 	// Clean previous networks
 	bIsLoaded = false;
-	const TArray<uint8>& ModelBuffer{ InferenceModel->GetFormatDesc().Data };
+	TConstArrayView<uint8> ModelBuffer = TConstArrayView<uint8>(&(ModelData.GetData()[GuidSize + VersionSize]), ModelData.Num() - GuidSize - VersionSize);
 
 	// Checking Inference Model 
 	{
diff --git a/Engine/Plugins/Experimental/NNX/Source/NNXRuntimeCPU/Public/NNXRuntimeCPU.h b/Engine/Plugins/Experimental/NNX/Source/NNXRuntimeCPU/Public/NNXRuntimeCPU.h
--- a/Engine/Plugins/Experimental/NNX/Source/NNXRuntimeCPU/Public/NNXRuntimeCPU.h
+++ b/Engine/Plugins/Experimental/NNX/Source/NNXRuntimeCPU/Public/NNXRuntimeCPU.h
@@ -35,10 +35,13 @@ namespace NNX
 		bool Init();
 		virtual ~FRuntimeCPU();
 		virtual FString GetRuntimeName() const;
-		virtual TUniquePtr<IModelOptimizer> CreateModelOptimizer() const;
 		virtual EMLRuntimeSupportFlags GetSupportFlags() const;
-		virtual FMLInferenceModel* CreateInferenceModel(UMLInferenceModel* Model);
-		FMLInferenceModel* CreateInferenceModel(UMLInferenceModel* InModel, const FMLInferenceNNXCPUConf& InConf);
+		virtual bool CanCreateModelData(FString FileType, TConstArrayView<uint8> FileData) const;
+		virtual TArray<uint8> CreateModelData(FString FileType, TConstArrayView<uint8> FileData);
+		virtual bool CanCreateModel(TConstArrayView<uint8> ModelData) const;
+		virtual TSharedPtr<FMLInferenceModel> CreateModel(TConstArrayView<uint8> ModelData);
+		static FGuid GUID;
+		static int32 Version;
 	};
 
 	static TUniquePtr<FRuntimeCPU> GCPURuntime;
@@ -81,7 +84,7 @@ namespace NNX
 		FMLInferenceModelCPU(Ort::Env* InORTEnvironment, const FMLInferenceNNXCPUConf& InNNXCPUConfiguration);
 		virtual ~FMLInferenceModelCPU() {};
 
-		bool Init(const UMLInferenceModel* InferenceModel);
+		bool Init(TConstArrayView<uint8> ModelData);
 		bool IsLoaded() const;
 
 		virtual int SetInputTensorShapes(TConstArrayView<FTensorShape> InInputShapes) override;
diff --git a/Engine/Plugins/Experimental/NNX/Source/NNXRuntimeORT/Private/NNXRuntimeORT.cpp b/Engine/Plugins/Experimental/NNX/Source/NNXRuntimeORT/Private/NNXRuntimeORT.cpp
--- a/Engine/Plugins/Experimental/NNX/Source/NNXRuntimeORT/Private/NNXRuntimeORT.cpp
+++ b/Engine/Plugins/Experimental/NNX/Source/NNXRuntimeORT/Private/NNXRuntimeORT.cpp
@@ -7,6 +7,7 @@
 #include "NNEProfilingTimer.h"
 #include "NNXRuntimeORTUtils.h"
 #include "NNXModelOptimizer.h"
+#include "NNECoreAttributeMap.h"
 #include "RedirectCoutAndCerrToUeLog.h"
 
 // NOTE: For now we only have DML on Windows, we should add support for XSX
@@ -58,78 +59,105 @@ EMLRuntimeSupportFlags FRuntimeORTDml::GetSupportFlags() const
 }
 #endif
 
-TUniquePtr<IModelOptimizer> FRuntimeORTCpu::CreateModelOptimizer() const
-{
-	return CreateONNXToONNXModelOptimizer();
-}
 
-#if PLATFORM_WINDOWS
-TUniquePtr<IModelOptimizer> FRuntimeORTCuda::CreateModelOptimizer() const
-{
-	return CreateONNXToONNXModelOptimizer();
-}
+FGuid FRuntimeORT::GUID = FGuid((int32)'O', (int32)'N', (int32)'N', (int32)'X');
+int32 FRuntimeORT::Version	= 0x00000001;
 
-TUniquePtr<IModelOptimizer> FRuntimeORTDml::CreateModelOptimizer() const
+bool FRuntimeORT::CanCreateModelData(FString FileType, TConstArrayView<uint8> FileData) const
 {
-	return CreateONNXToONNXModelOptimizer();
+	return FileType.Compare("onnx", ESearchCase::IgnoreCase) == 0;
 }
-#endif
 
-FMLInferenceModel* FRuntimeORTCpu::CreateInferenceModel(UMLInferenceModel* InModel, const FMLInferenceNNXORTConf& InConf)
+TArray<uint8> FRuntimeORT::CreateModelData(FString FileType, TConstArrayView<uint8> FileData)
 {
-	FMLInferenceModelORTCpu* ORTModel = new FMLInferenceModelORTCpu(&NNXEnvironmentORT, InConf);
-	if (!ORTModel->Init(InModel))
+	if (!CanCreateModelData(FileType, FileData))
 	{
-		delete ORTModel;
-		ORTModel = nullptr;
+		return {};
 	}
 
-	return ORTModel;
-}
+	TUniquePtr<IModelOptimizer> Optimizer = CreateONNXToONNXModelOptimizer();
 
-#if PLATFORM_WINDOWS
-FMLInferenceModel* FRuntimeORTCuda::CreateInferenceModel(UMLInferenceModel* InModel, const FMLInferenceNNXORTConf& InConf)
-{
-	FMLInferenceModelORTCuda* ORTModel = new FMLInferenceModelORTCuda(&NNXEnvironmentORT, InConf);
-	if (!ORTModel->Init(InModel))
+	FNNIModelRaw InputModel;
+	InputModel.Data = FileData;
+	InputModel.Format = ENNXInferenceFormat::ONNX;
+	FNNIModelRaw OutputModel;
+	FOptimizerOptionsMap Options;
+	if (!Optimizer->Optimize(InputModel, OutputModel, Options))
 	{
-		delete ORTModel;
-		ORTModel = nullptr;
+		return {};
 	}
 
-	return ORTModel;
+	int32 GuidSize = sizeof(FRuntimeORT::GUID);
+	int32 VersionSize = sizeof(FRuntimeORT::Version);
+	TArray<uint8> Result;
+	FMemoryWriter Writer(Result);
+	Writer << FRuntimeORT::GUID;
+	Writer << FRuntimeORT::Version;
+	Writer.Serialize(OutputModel.Data.GetData(), OutputModel.Data.Num());
+	return Result;
 }
 
-FMLInferenceModel* FRuntimeORTDml::CreateInferenceModel(UMLInferenceModel* InModel, const FMLInferenceNNXORTConf& InConf)
+
+bool FRuntimeORT::CanCreateModel(TConstArrayView<uint8> ModelData) const
 {
-	FMLInferenceModelORTDml* ORTModel = new FMLInferenceModelORTDml(&NNXEnvironmentORT, InConf);
-	if (!ORTModel->Init(InModel))
+	int32 GuidSize = sizeof(FRuntimeORT::GUID);
+	int32 VersionSize = sizeof(FRuntimeORT::Version);
+	if (ModelData.Num() <= GuidSize + VersionSize)
 	{
-		delete ORTModel;
-		ORTModel = nullptr;
+		return false;
 	}
-
-	return ORTModel;
+	bool bResult = FGenericPlatformMemory::Memcmp(&(ModelData[0]), &(FRuntimeORT::GUID), GuidSize) == 0;
+	bResult &= FGenericPlatformMemory::Memcmp(&(ModelData[GuidSize]), &(FRuntimeORT::Version), VersionSize) == 0;
+	return bResult;
 }
-#endif
 
-FMLInferenceModel* FRuntimeORTCpu::CreateInferenceModel(UMLInferenceModel* InModel)
+TSharedPtr<FMLInferenceModel> FRuntimeORTCpu::CreateModel(TConstArrayView<uint8> ModelData)
 {
-	FMLInferenceNNXORTConf ORTInferenceConf;
-	return CreateInferenceModel(InModel, ORTInferenceConf);
+	if (!CanCreateModel(ModelData))
+	{
+		return TSharedPtr<FMLInferenceModel>();
+	}
+	const FMLInferenceNNXORTConf InConf;
+	FMLInferenceModelORTCpu* Model = new FMLInferenceModelORTCpu(&NNXEnvironmentORT, InConf);
+	if (!Model->Init(ModelData))
+	{
+		delete Model;
+		return TSharedPtr<FMLInferenceModel>();
+	}
+	return TSharedPtr<FMLInferenceModel>(Model);
 }
 
 #if PLATFORM_WINDOWS
-FMLInferenceModel* FRuntimeORTCuda::CreateInferenceModel(UMLInferenceModel* InModel)
+TSharedPtr<FMLInferenceModel> FRuntimeORTCuda::CreateModel(TConstArrayView<uint8> ModelData)
 {
-	FMLInferenceNNXORTConf ORTInferenceConf;
-	return CreateInferenceModel(InModel, ORTInferenceConf);
+	if (!CanCreateModel(ModelData))
+	{
+		return TSharedPtr<FMLInferenceModel>();
+	}
+	const FMLInferenceNNXORTConf InConf;
+	FMLInferenceModelORTCuda* Model = new FMLInferenceModelORTCuda(&NNXEnvironmentORT, InConf);
+	if (!Model->Init(ModelData))
+	{
+		delete Model;
+		return TSharedPtr<FMLInferenceModel>();
+	}
+	return TSharedPtr<FMLInferenceModel>(Model);
 }
 
-FMLInferenceModel* FRuntimeORTDml::CreateInferenceModel(UMLInferenceModel* InModel)
+TSharedPtr<FMLInferenceModel> FRuntimeORTDml::CreateModel(TConstArrayView<uint8> ModelData)
 {
-	FMLInferenceNNXORTConf ORTInferenceConf;
-	return CreateInferenceModel(InModel, ORTInferenceConf);
+	if (!CanCreateModel(ModelData))
+	{
+		return TSharedPtr<FMLInferenceModel>();
+	}
+	const FMLInferenceNNXORTConf InConf;
+	FMLInferenceModelORTDml* Model = new FMLInferenceModelORTDml(&NNXEnvironmentORT, InConf);
+	if (!Model->Init(ModelData))
+	{
+		delete Model;
+		return TSharedPtr<FMLInferenceModel>();
+	}
+	return TSharedPtr<FMLInferenceModel>(Model);
 }
 #endif
 
@@ -144,13 +172,17 @@ FMLInferenceModelORT::FMLInferenceModelORT(
 	ORTConfiguration(InORTConfiguration)
 { }
 
-bool FMLInferenceModelORT::Init(const UMLInferenceModel* InferenceModel)
+bool FMLInferenceModelORT::Init(TConstArrayView<uint8> ModelData)
 {
 	DECLARE_SCOPE_CYCLE_COUNTER(TEXT("FMLInferenceModelORT_Init"), STAT_FMLInferenceModelORT_Init, STATGROUP_MachineLearning);
-	
+
+	// Get the header size
+	int32 GuidSize = sizeof(FRuntimeORT::GUID);
+	int32 VersionSize = sizeof(FRuntimeORT::Version);
+
 	// Clean previous networks
 	bIsLoaded = false;
-	const TArray<uint8>& ModelBuffer{ InferenceModel->GetFormatDesc().Data };
+	TConstArrayView<uint8> ModelBuffer = TConstArrayView<uint8>(&(ModelData.GetData()[GuidSize + VersionSize]), ModelData.Num() - GuidSize - VersionSize);
 
 	// Checking Inference Model 
 	{
diff --git a/Engine/Plugins/Experimental/NNX/Source/NNXRuntimeORT/Private/NNXRuntimeORT.h b/Engine/Plugins/Experimental/NNX/Source/NNXRuntimeORT/Private/NNXRuntimeORT.h
--- a/Engine/Plugins/Experimental/NNX/Source/NNXRuntimeORT/Private/NNXRuntimeORT.h
+++ b/Engine/Plugins/Experimental/NNX/Source/NNXRuntimeORT/Private/NNXRuntimeORT.h
@@ -50,44 +50,49 @@ namespace NNX
 		GraphOptimizationLevel OptimizationLevel = GraphOptimizationLevel::ORT_ENABLE_ALL;
 	};
 
-	class FRuntimeORTCpu : public IRuntime
+	class FRuntimeORT : public IRuntime
+	{
+	public:
+		static FGuid GUID;
+		static int32 Version;
+
+		virtual bool CanCreateModelData(FString FileType, TConstArrayView<uint8> FileData) const;
+		virtual TArray<uint8> CreateModelData(FString FileType, TConstArrayView<uint8> FileData);
+		virtual bool CanCreateModel(TConstArrayView<uint8> ModelData) const;
+	};
+
+	class FRuntimeORTCpu : public FRuntimeORT
 	{
 	public:
 		Ort::Env NNXEnvironmentORT;
 		FRuntimeORTCpu() = default;
 		virtual ~FRuntimeORTCpu() = default;
 		virtual FString GetRuntimeName() const override;
-		virtual TUniquePtr<IModelOptimizer> CreateModelOptimizer() const override;
 		virtual EMLRuntimeSupportFlags GetSupportFlags() const override;
-		virtual FMLInferenceModel* CreateInferenceModel(UMLInferenceModel* InModel);
-		FMLInferenceModel* CreateInferenceModel(UMLInferenceModel* InModel, const FMLInferenceNNXORTConf& InConf);
+		virtual TSharedPtr<FMLInferenceModel> CreateModel(TConstArrayView<uint8> ModelData);
 	};
 
 #if PLATFORM_WINDOWS
-	class FRuntimeORTCuda : public IRuntime
+	class FRuntimeORTCuda : public FRuntimeORT
 	{
 	public:
 		Ort::Env NNXEnvironmentORT;
 		FRuntimeORTCuda() = default;
 		virtual ~FRuntimeORTCuda() = default;
 		virtual FString GetRuntimeName() const override;
-		virtual TUniquePtr<IModelOptimizer> CreateModelOptimizer() const override;
 		virtual EMLRuntimeSupportFlags GetSupportFlags() const override;
-		virtual FMLInferenceModel* CreateInferenceModel(UMLInferenceModel* InModel);
-		FMLInferenceModel* CreateInferenceModel(UMLInferenceModel* InModel, const FMLInferenceNNXORTConf& InConf);
+		virtual TSharedPtr<FMLInferenceModel> CreateModel(TConstArrayView<uint8> ModelData);
 	};
 
-	class FRuntimeORTDml : public IRuntime
+	class FRuntimeORTDml : public FRuntimeORT
 	{
 	public:
 		Ort::Env NNXEnvironmentORT;
 		FRuntimeORTDml() = default;
 		virtual ~FRuntimeORTDml() = default;
 		virtual FString GetRuntimeName() const override;
-		virtual TUniquePtr<IModelOptimizer> CreateModelOptimizer() const override;
 		virtual EMLRuntimeSupportFlags GetSupportFlags() const override;
-		virtual FMLInferenceModel* CreateInferenceModel(UMLInferenceModel* InModel);
-		FMLInferenceModel* CreateInferenceModel(UMLInferenceModel* InModel, const FMLInferenceNNXORTConf& InConf);
+		virtual TSharedPtr<FMLInferenceModel> CreateModel(TConstArrayView<uint8> ModelData);
 	};
 #endif
 
@@ -180,7 +185,7 @@ namespace NNX
 	public:
 		virtual ~FMLInferenceModelORT() = default;
 
-		bool Init(const UMLInferenceModel* InferenceModel);
+		bool Init(TConstArrayView<uint8> ModelData);
 		bool IsLoaded() const;
 
 		virtual int SetInputTensorShapes(TConstArrayView<FTensorShape> InInputShapes) override;
diff --git a/Engine/Plugins/Experimental/NNX/Source/NNXRuntimeRDG/Private/NNXRuntimeDML.cpp b/Engine/Plugins/Experimental/NNX/Source/NNXRuntimeRDG/Private/NNXRuntimeDML.cpp
--- a/Engine/Plugins/Experimental/NNX/Source/NNXRuntimeRDG/Private/NNXRuntimeDML.cpp
+++ b/Engine/Plugins/Experimental/NNX/Source/NNXRuntimeRDG/Private/NNXRuntimeDML.cpp
@@ -1172,7 +1172,7 @@ public:
 	FMLInferenceModelDml();
 	~FMLInferenceModelDml();
 
-	bool Init(UMLInferenceModel* InModel, FDeviceContextDml* InDevCtx);
+	bool Init(TConstArrayView<uint8> ModelData, FDeviceContextDml* InDevCtx);
 
 protected:
 
@@ -1190,7 +1190,7 @@ private:
 //
 //
 //
-class FMLRuntimeDml : public IRuntime
+class FMLRuntimeDml : public FMLRuntimeRDG
 {
 public:
 
@@ -1201,32 +1201,27 @@ public:
 	{
 		return NNX_RUNTIME_DML_NAME;
 	}
-
-	virtual TUniquePtr<IModelOptimizer> CreateModelOptimizer() const
-	{
-		return CreateONNXToNNXModelOptimizer();
-		
-		//TODO jira 167581: Add DML specific model validator
-		//see ModelOptimizer->AddValidator(MakeShared<FModelValidatorHlsl>());
-		//in NNXRuntimeHLSL.cpp
-	}
 	
 	virtual EMLRuntimeSupportFlags GetSupportFlags() const override
 	{
 		return EMLRuntimeSupportFlags::RDG;
 	}
 
-	virtual FMLInferenceModel* CreateInferenceModel(UMLInferenceModel* Model)
+	virtual TSharedPtr<FMLInferenceModel> CreateModel(TConstArrayView<uint8> ModelData) override
 	{
-		FMLInferenceModelDml* DmlModel = new FMLInferenceModelDml();
-
-		if (!DmlModel->Init(Model, &Ctx))
+		if (!CanCreateModel(ModelData))
 		{
-			delete DmlModel;
-			DmlModel = nullptr;
+			return TSharedPtr<FMLInferenceModel>();
 		}
 
-		return DmlModel;
+		// Create the model and initialize it with the data not including the header
+		FMLInferenceModelDml* Model = new FMLInferenceModelDml();
+		if (!Model->Init(ModelData, &Ctx))
+		{
+			delete Model;
+			return TSharedPtr<FMLInferenceModel>();
+		}
+		return TSharedPtr<FMLInferenceModel>(Model);
 	}
 
 	bool Init();
@@ -1466,12 +1461,12 @@ FMLInferenceModelDml::~FMLInferenceModelDml()
 //
 //
 //
-bool FMLInferenceModelDml::Init(UMLInferenceModel* InModel, FDeviceContextDml* InDevCtx)
+bool FMLInferenceModelDml::Init(TConstArrayView<uint8> ModelData, FDeviceContextDml* InDevCtx)
 {
-	check(InModel != nullptr);
+	check(ModelData.Num() > 0);
 	FMLRuntimeFormat	Format;
 
-	if (!LoadModel(InModel->GetFormatDesc(), Format))
+	if (!LoadModel(ModelData, Format))
 	{
 		return false;
 	}
diff --git a/Engine/Plugins/Experimental/NNX/Source/NNXRuntimeRDG/Private/NNXRuntimeHLSL.cpp b/Engine/Plugins/Experimental/NNX/Source/NNXRuntimeRDG/Private/NNXRuntimeHLSL.cpp
--- a/Engine/Plugins/Experimental/NNX/Source/NNXRuntimeRDG/Private/NNXRuntimeHLSL.cpp
+++ b/Engine/Plugins/Experimental/NNX/Source/NNXRuntimeRDG/Private/NNXRuntimeHLSL.cpp
@@ -30,7 +30,7 @@ public:
 	FMLInferenceModelHlsl();
 	~FMLInferenceModelHlsl();
 
-	bool Init(UMLInferenceModel* InModel);
+	bool Init(TConstArrayView<uint8> ModelData);
 
 protected:
 
@@ -63,12 +63,12 @@ FMLInferenceModelHlsl::~FMLInferenceModelHlsl()
 //
 //
 //
-bool FMLInferenceModelHlsl::Init(UMLInferenceModel* InModel)
+bool FMLInferenceModelHlsl::Init(TConstArrayView<uint8> ModelData)
 {
-	check(InModel != nullptr);
+	check(ModelData.Num() > 0);
 	FMLRuntimeFormat	Format;
 
-	if (!LoadModel(InModel->GetFormatDesc(), Format))
+	if (!LoadModel(ModelData, Format))
 	{
 		return false;
 	}
@@ -264,7 +264,7 @@ void FMLInferenceModelHlsl::AddDispatchOps_RenderThread(FRDGBuilder& GraphBuilde
 //
 //
 //
-class FMLRuntimeHlsl : public IRuntime
+class FMLRuntimeHlsl : public FMLRuntimeRDG
 {
 public:
 
@@ -295,33 +295,27 @@ public:
 		return NNX_RUNTIME_HLSL_NAME;
 	}
 
-	virtual TUniquePtr<IModelOptimizer> CreateModelOptimizer() const
-	{
-		TUniquePtr<IModelOptimizer> ModelOptimizer = CreateONNXToNNXModelOptimizer();
-		check(ModelOptimizer.IsValid());
-		
-		ModelOptimizer->AddValidator(MakeShared<FModelValidatorHlsl>());
-		
-		return ModelOptimizer;
-	}
-
 	// Returns flags from ERuntimeSupportFlags
 	virtual EMLRuntimeSupportFlags GetSupportFlags() const override
 	{
 		return EMLRuntimeSupportFlags::RDG;
 	}
 
-	virtual FMLInferenceModel* CreateInferenceModel(UMLInferenceModel* Model)
+	virtual TSharedPtr<FMLInferenceModel> CreateModel(TConstArrayView<uint8> ModelData) override
 	{
-		auto HlslModel = new FMLInferenceModelHlsl();
-
-		if (!HlslModel->Init(Model))
+		if (!CanCreateModel(ModelData))
 		{
-			delete HlslModel;
-			HlslModel = nullptr;
+			return TSharedPtr<FMLInferenceModel>();
 		}
 
-		return HlslModel;
+		// Create the model and initialize it with the data not including the header
+		FMLInferenceModelHlsl* Model = new FMLInferenceModelHlsl();
+		if (!Model->Init(ModelData))
+		{
+			delete Model;
+			return TSharedPtr<FMLInferenceModel>();
+		}
+		return TSharedPtr<FMLInferenceModel>(Model);
 	}
 };
 
diff --git a/Engine/Plugins/Experimental/NNX/Source/NNXRuntimeRDG/Private/NNXRuntimeRDG.cpp b/Engine/Plugins/Experimental/NNX/Source/NNXRuntimeRDG/Private/NNXRuntimeRDG.cpp
--- a/Engine/Plugins/Experimental/NNX/Source/NNXRuntimeRDG/Private/NNXRuntimeRDG.cpp
+++ b/Engine/Plugins/Experimental/NNX/Source/NNXRuntimeRDG/Private/NNXRuntimeRDG.cpp
@@ -3,6 +3,7 @@
 #include "NNXRuntimeRDG.h"
 #include "NNXInferenceModel.h"
 #include "NNXRuntimeFormat.h"
+#include "NNXModelOptimizer.h"
 
 #include "RenderGraphBuilder.h"
 #include "RenderGraphUtils.h"
@@ -13,6 +14,9 @@
 namespace NNX
 {
 
+FGuid FMLRuntimeRDG::GUID = FGuid((int32)'R', (int32)'R', (int32)'D', (int32)'G');
+int32 FMLRuntimeRDG::Version = 0x00000001;
+
 struct FMLInferenceModelRDG::FReadbackEntry
 {
 	TUniquePtr<FRHIGPUBufferReadback>	RHI;
@@ -169,6 +173,53 @@ bool FAttributeValidator::Validate(const UE::NNECore::FAttributeMap& AttributesT
 
 	return bAreAttributesValid;
 }
+
+bool FMLRuntimeRDG::CanCreateModelData(FString FileType, TConstArrayView<uint8> FileData) const
+{
+	return FileType.Compare("onnx", ESearchCase::IgnoreCase) == 0;
+}
+
+TArray<uint8> FMLRuntimeRDG::CreateModelData(FString FileType, TConstArrayView<uint8> FileData)
+{
+	if (!CanCreateModelData(FileType, FileData))
+	{
+		return {};
+	}
+
+	TUniquePtr<IModelOptimizer> Optimizer = CreateONNXToNNXModelOptimizer();
+
+	FNNIModelRaw InputModel;
+	InputModel.Data = FileData;
+	InputModel.Format = ENNXInferenceFormat::ONNX;
+	FNNIModelRaw OutputModel;
+	FOptimizerOptionsMap Options;
+	if (!Optimizer->Optimize(InputModel, OutputModel, Options))
+	{
+		return {};
+	}
+
+	int32 GuidSize = sizeof(FMLRuntimeRDG::GUID);
+	int32 VersionSize = sizeof(FMLRuntimeRDG::Version);
+	TArray<uint8> Result;
+	FMemoryWriter Writer(Result);
+	Writer << FMLRuntimeRDG::GUID;
+	Writer << FMLRuntimeRDG::Version;
+	Writer.Serialize(OutputModel.Data.GetData(), OutputModel.Data.Num());
+	return Result;
+};
+
+bool FMLRuntimeRDG::CanCreateModel(TConstArrayView<uint8> ModelData) const
+{
+	int32 GuidSize = sizeof(FMLRuntimeRDG::GUID);
+	int32 VersionSize = sizeof(FMLRuntimeRDG::Version);
+	if (ModelData.Num() <= GuidSize + VersionSize)
+	{
+		return false;
+	}
+	bool bResult = FGenericPlatformMemory::Memcmp(&(ModelData[0]), &(FMLRuntimeRDG::GUID), GuidSize) == 0;
+	bResult &= FGenericPlatformMemory::Memcmp(&(ModelData[GuidSize]), &(FMLRuntimeRDG::Version), VersionSize) == 0;
+	return bResult;
+};
 	
 //
 //
@@ -188,17 +239,13 @@ FMLInferenceModelRDG::~FMLInferenceModelRDG()
 //
 //
 //
-bool FMLInferenceModelRDG::LoadModel(const FNNIModelRaw& InModel, FMLRuntimeFormat& Format)
+bool FMLInferenceModelRDG::LoadModel(TConstArrayView<uint8> ModelData, FMLRuntimeFormat& Format)
 {
-	ENNXInferenceFormat FormatType = InModel.Format;
-
-	if (FormatType != ENNXInferenceFormat::NNXRT)
-	{
-		UE_LOG(LogNNX, Warning, TEXT("Unsupported format type for NNX inference model"));
-		return false;
-	}
+	int32 GuidSize = sizeof(FMLRuntimeRDG::GUID);
+	int32 VersionSize = sizeof(FMLRuntimeRDG::Version);
+	TArray<uint8> ModelBuffer = TArray<uint8>(&(ModelData.GetData()[GuidSize + VersionSize]), ModelData.Num() - GuidSize - VersionSize);
 
-	FMemoryReader Reader(InModel.Data);
+	FMemoryReader Reader(ModelBuffer);
 
 	FMLRuntimeFormat::StaticStruct()->SerializeBin(Reader, &Format);
 
diff --git a/Engine/Plugins/Experimental/NNX/Source/NNXRuntimeRDG/Private/NNXRuntimeRDG.h b/Engine/Plugins/Experimental/NNX/Source/NNXRuntimeRDG/Private/NNXRuntimeRDG.h
--- a/Engine/Plugins/Experimental/NNX/Source/NNXRuntimeRDG/Private/NNXRuntimeRDG.h
+++ b/Engine/Plugins/Experimental/NNX/Source/NNXRuntimeRDG/Private/NNXRuntimeRDG.h
@@ -68,6 +68,17 @@ public:
 using FTensorRDGArray = TArray<FTensorRDG, TInlineAllocator<16>>;
 using FIntArray = TArray<int32, TInlineAllocator<16>>;
 
+class FMLRuntimeRDG : public IRuntime
+{
+public:
+	static FGuid GUID;
+	static int32 Version;
+
+	virtual bool CanCreateModelData(FString FileType, TConstArrayView<uint8> FileData) const;
+	virtual TArray<uint8> CreateModelData(FString FileType, TConstArrayView<uint8> FileData);
+	virtual bool CanCreateModel(TConstArrayView<uint8> ModelData) const;
+};
+
 /** 
  * RDG inference model base class
  */
@@ -88,7 +99,7 @@ protected:
 
 	FMLInferenceModelRDG();
 
-	bool LoadModel(const FNNIModelRaw& InModel, FMLRuntimeFormat& Format);
+	bool LoadModel(TConstArrayView<uint8> ModelData, FMLRuntimeFormat& Format);
 
 	int SetTensors(FRDGBuilder& GraphBuilder, FTensorRDGArray& OutTensorRDGs, FIntArray& OutIndices, TConstArrayView<FMLTensorBinding> InBindings, TConstArrayView<FTensorDesc> InTensorDescs, TConstArrayView<FTensorShape> InTensorShapes);
 	
diff --git a/Engine/Plugins/Experimental/NNX/Source/NNXUtils/Private/NNXModelOptimizer.cpp b/Engine/Plugins/Experimental/NNX/Source/NNXUtils/Private/NNXModelOptimizer.cpp
--- a/Engine/Plugins/Experimental/NNX/Source/NNXUtils/Private/NNXModelOptimizer.cpp
+++ b/Engine/Plugins/Experimental/NNX/Source/NNXUtils/Private/NNXModelOptimizer.cpp
@@ -659,35 +659,4 @@ NNXUTILS_API TUniquePtr<IModelOptimizer> CreateModelOptimizer(ENNXInferenceForma
 	return nullptr;
 }
 
-/** Helper to create an optimized model for a given runtime from ONNX */
-NNXUTILS_API bool CreateRuntimeModelFromONNX(const FNNIModelRaw& ONNXModel, FNNIModelRaw& OptimizedModel, const FString& RuntimeName, const FOptimizerOptionsMap& Options)
-{
-	OptimizedModel = FNNIModelRaw{};
-	
-	//TODO jira 167594: Register the NNEngine even when CreateInferenceModel can't be called
-	//on the platform because of missing hardware support. This would allow
-	//to optimize a model even when the current hardware cannot run inference on current hardware.
-	NNX::IRuntime* Runtime = NNX::GetRuntime(RuntimeName);
-	if (!Runtime)
-	{
-		UE_LOG(LogNNX, Warning, TEXT("Can't find runtime %s!"), *RuntimeName);
-		return false;
-	}
-
-	TUniquePtr<IModelOptimizer> Optimizer = Runtime->CreateModelOptimizer();
-	if (!Optimizer.IsValid())
-	{
-		UE_LOG(LogNNX, Warning, TEXT("No optimizer available for runtime %s!"), *RuntimeName);
-		return false;
-	}
-
-	if (!Optimizer->Optimize(ONNXModel, OptimizedModel, Options))
-	{
-		UE_LOG(LogNNX, Warning, TEXT("Could not create an optimized model for the required runtime %s. Check the log!"), *RuntimeName);
-		return false;
-	}
-
-	return true;
-}
-
 } // namespace NNX
diff --git a/Engine/Plugins/Experimental/NNX/Source/NNXUtils/Public/NNXModelOptimizer.h b/Engine/Plugins/Experimental/NNX/Source/NNXUtils/Public/NNXModelOptimizer.h
--- a/Engine/Plugins/Experimental/NNX/Source/NNXUtils/Public/NNXModelOptimizer.h
+++ b/Engine/Plugins/Experimental/NNX/Source/NNXUtils/Public/NNXModelOptimizer.h
@@ -29,7 +29,4 @@ inline TUniquePtr<IModelOptimizer> CreateONNXToONNXModelOptimizer()
 	return CreateModelOptimizer(ENNXInferenceFormat::ONNX, ENNXInferenceFormat::ONNX);
 }
 
-/** Helper to create an optimized model for a given runtime from ONNX */
-NNXUTILS_API bool CreateRuntimeModelFromONNX(FNNIModelRaw& OutputModel, const FNNIModelRaw& ONNXModel, const FString& RuntimeName, const FOptimizerOptionsMap& Options);
-
 } // NNX
