NNE - Changing to unique pointers to return models

diff --git a/Engine/Plugins/Experimental/NNX/Source/NNXCore/Private/NNEModel.cpp b/Engine/Plugins/Experimental/NNX/Source/NNXCore/Private/NNEModel.cpp
--- a/Engine/Plugins/Experimental/NNX/Source/NNXCore/Private/NNEModel.cpp
+++ b/Engine/Plugins/Experimental/NNX/Source/NNXCore/Private/NNEModel.cpp
@@ -162,12 +162,12 @@ bool UNNEModel::SetInputOutput(const TArray<FNNETensor>& Input, UPARAM(ref) TArr
 	return true;
 }
 
-bool UNNEModel::Run()
+bool UNNEModel::RunSync()
 {
 	if (!Model.IsValid())
 	{
 		return false;
 	}
 
-	return Model->Run(InputBindings, OutputBindings) == 0;
+	return Model->RunSync(InputBindings, OutputBindings) == 0;
 }
\ No newline at end of file
diff --git a/Engine/Plugins/Experimental/NNX/Source/NNXCore/Public/NNEModel.h b/Engine/Plugins/Experimental/NNX/Source/NNXCore/Public/NNEModel.h
--- a/Engine/Plugins/Experimental/NNX/Source/NNXCore/Public/NNEModel.h
+++ b/Engine/Plugins/Experimental/NNX/Source/NNXCore/Public/NNEModel.h
@@ -77,11 +77,11 @@ public:
 	 * Upon successful return of this function, the data inside the output tensors passed to SetInputOutput will contain the resulting data.
 	 */
 	UFUNCTION(BlueprintCallable, Category = "NNE - Neural Network Engine")
-	bool Run();
+	bool RunSync();
 	
 private:
 
-	TSharedPtr<NNX::FMLInferenceModel> Model;
+	TUniquePtr<NNX::FMLInferenceModel> Model;
 
 	TArray<NNX::FMLTensorBinding> InputBindings;
 	TArray<NNX::FMLTensorBinding> OutputBindings;
diff --git a/Engine/Plugins/Experimental/NNX/Source/NNXCore/Public/NNXRuntime.h b/Engine/Plugins/Experimental/NNX/Source/NNXCore/Public/NNXRuntime.h
--- a/Engine/Plugins/Experimental/NNX/Source/NNXCore/Public/NNXRuntime.h
+++ b/Engine/Plugins/Experimental/NNX/Source/NNXCore/Public/NNXRuntime.h
@@ -61,7 +61,7 @@ public:
 
 	// Model factory interface
 	virtual bool CanCreateModel(TConstArrayView<uint8> ModelData) const = 0;
-	virtual TSharedPtr<FMLInferenceModel> CreateModel(TConstArrayView<uint8> ModelData) = 0;
+	virtual TUniquePtr<FMLInferenceModel> CreateModel(TConstArrayView<uint8> ModelData) = 0;
 };
 
 /** Tensor memory binding type */
@@ -129,7 +129,7 @@ public:
 
 	/** This call is synchronous on all engine types(CPU, RDG, GPU), i.e.the calling thread will be blocked until execution is finished.
 	* Bindings should point to a buffers big enough */
-	virtual int Run(TConstArrayView<FMLTensorBinding> InInputTensors,
+	virtual int RunSync(TConstArrayView<FMLTensorBinding> InInputTensors,
 					TConstArrayView<FMLTensorBinding> InOutputTensors) = 0;
 
 	/** Enqueue the execution on the Render graph render thread. It's caller's responsibility to actually run the graph.
diff --git a/Engine/Plugins/Experimental/NNX/Source/NNXQA/Private/NNXQAUtils.cpp b/Engine/Plugins/Experimental/NNX/Source/NNXQA/Private/NNXQAUtils.cpp
--- a/Engine/Plugins/Experimental/NNX/Source/NNXQA/Private/NNXQAUtils.cpp
+++ b/Engine/Plugins/Experimental/NNX/Source/NNXQA/Private/NNXQAUtils.cpp
@@ -363,7 +363,7 @@ namespace Test
 	{
 		OutOutputMemBuffers.Empty();
 
-		TSharedPtr<FMLInferenceModel> InferenceModel = Runtime->CreateModel(Runtime->CreateModelData(FString("onnx"), ONNXModelData.Data));
+		TUniquePtr<FMLInferenceModel> InferenceModel = Runtime->CreateModel(Runtime->CreateModelData(FString("onnx"), ONNXModelData.Data));
 		
 		if (!InferenceModel.IsValid())
 		{
@@ -416,7 +416,7 @@ namespace Test
 		}
 
 		// Run inference
-		if (0 != InferenceModel->Run(InputBindings, OutputBindings))
+		if (0 != InferenceModel->RunSync(InputBindings, OutputBindings))
 		{
 			UE_LOG(LogNNX, Error, TEXT("Failed to run the model."));
 			return false;
diff --git a/Engine/Plugins/Experimental/NNX/Source/NNXRuntimeCPU/Private/NNXRuntimeCPU.cpp b/Engine/Plugins/Experimental/NNX/Source/NNXRuntimeCPU/Private/NNXRuntimeCPU.cpp
--- a/Engine/Plugins/Experimental/NNX/Source/NNXRuntimeCPU/Private/NNXRuntimeCPU.cpp
+++ b/Engine/Plugins/Experimental/NNX/Source/NNXRuntimeCPU/Private/NNXRuntimeCPU.cpp
@@ -68,11 +68,11 @@ bool FRuntimeCPU::CanCreateModel(TConstArrayView<uint8> ModelData) const
 	return bResult;
 }
 
-TSharedPtr<FMLInferenceModel> FRuntimeCPU::CreateModel(TConstArrayView<uint8> ModelData)
+TUniquePtr<FMLInferenceModel> FRuntimeCPU::CreateModel(TConstArrayView<uint8> ModelData)
 {
 	if (!CanCreateModel(ModelData))
 	{
-		return TSharedPtr<FMLInferenceModel>();
+		return TUniquePtr<FMLInferenceModel>();
 	}
 
 	// Get the header size
@@ -85,9 +85,9 @@ TSharedPtr<FMLInferenceModel> FRuntimeCPU::CreateModel(TConstArrayView<uint8> Mo
 	if (!Model->Init(ModelData))
 	{
 		delete Model;
-		return TSharedPtr<FMLInferenceModel>();
+		return TUniquePtr<FMLInferenceModel>();
 	}
-	return TSharedPtr<FMLInferenceModel>(Model);
+	return TUniquePtr<FMLInferenceModel>(Model);
 }
 
 bool FRuntimeCPU::Init()
@@ -292,7 +292,7 @@ int FMLInferenceModelCPU::SetInputTensorShapes(TConstArrayView<FTensorShape> InI
 	return 0;
 }
 
-int FMLInferenceModelCPU::Run(TConstArrayView<FMLTensorBinding> InInputBindings, TConstArrayView<FMLTensorBinding> InOutputBindings)
+int FMLInferenceModelCPU::RunSync(TConstArrayView<FMLTensorBinding> InInputBindings, TConstArrayView<FMLTensorBinding> InOutputBindings)
 {
 	DECLARE_SCOPE_CYCLE_COUNTER(TEXT("FMLInferenceModelCPU_Run"), STAT_FMLInferenceModelCPU_Run, STATGROUP_MachineLearning);
 
diff --git a/Engine/Plugins/Experimental/NNX/Source/NNXRuntimeCPU/Public/NNXRuntimeCPU.h b/Engine/Plugins/Experimental/NNX/Source/NNXRuntimeCPU/Public/NNXRuntimeCPU.h
--- a/Engine/Plugins/Experimental/NNX/Source/NNXRuntimeCPU/Public/NNXRuntimeCPU.h
+++ b/Engine/Plugins/Experimental/NNX/Source/NNXRuntimeCPU/Public/NNXRuntimeCPU.h
@@ -39,7 +39,7 @@ namespace NNX
 		virtual bool CanCreateModelData(FString FileType, TConstArrayView<uint8> FileData) const;
 		virtual TArray<uint8> CreateModelData(FString FileType, TConstArrayView<uint8> FileData);
 		virtual bool CanCreateModel(TConstArrayView<uint8> ModelData) const;
-		virtual TSharedPtr<FMLInferenceModel> CreateModel(TConstArrayView<uint8> ModelData);
+		virtual TUniquePtr<FMLInferenceModel> CreateModel(TConstArrayView<uint8> ModelData);
 		static FGuid GUID;
 		static int32 Version;
 	};
@@ -88,7 +88,7 @@ namespace NNX
 		bool IsLoaded() const;
 
 		virtual int SetInputTensorShapes(TConstArrayView<FTensorShape> InInputShapes) override;
-		virtual int Run(TConstArrayView<FMLTensorBinding> InInputBindings, TConstArrayView<FMLTensorBinding> InOutputBindings) override;
+		virtual int RunSync(TConstArrayView<FMLTensorBinding> InInputBindings, TConstArrayView<FMLTensorBinding> InOutputBindings) override;
 
 		float GetLastRunTimeMSec() const;
 		UE::NNEProfiling::Internal::FStatistics GetRunStatistics() const;
diff --git a/Engine/Plugins/Experimental/NNX/Source/NNXRuntimeORT/Private/NNXRuntimeORT.cpp b/Engine/Plugins/Experimental/NNX/Source/NNXRuntimeORT/Private/NNXRuntimeORT.cpp
--- a/Engine/Plugins/Experimental/NNX/Source/NNXRuntimeORT/Private/NNXRuntimeORT.cpp
+++ b/Engine/Plugins/Experimental/NNX/Source/NNXRuntimeORT/Private/NNXRuntimeORT.cpp
@@ -111,53 +111,53 @@ bool FRuntimeORT::CanCreateModel(TConstArrayView<uint8> ModelData) const
 	return bResult;
 }
 
-TSharedPtr<FMLInferenceModel> FRuntimeORTCpu::CreateModel(TConstArrayView<uint8> ModelData)
+TUniquePtr<FMLInferenceModel> FRuntimeORTCpu::CreateModel(TConstArrayView<uint8> ModelData)
 {
 	if (!CanCreateModel(ModelData))
 	{
-		return TSharedPtr<FMLInferenceModel>();
+		return TUniquePtr<FMLInferenceModel>();
 	}
 	const FMLInferenceNNXORTConf InConf;
 	FMLInferenceModelORTCpu* Model = new FMLInferenceModelORTCpu(&NNXEnvironmentORT, InConf);
 	if (!Model->Init(ModelData))
 	{
 		delete Model;
-		return TSharedPtr<FMLInferenceModel>();
+		return TUniquePtr<FMLInferenceModel>();
 	}
-	return TSharedPtr<FMLInferenceModel>(Model);
+	return TUniquePtr<FMLInferenceModel>(Model);
 }
 
 #if PLATFORM_WINDOWS
-TSharedPtr<FMLInferenceModel> FRuntimeORTCuda::CreateModel(TConstArrayView<uint8> ModelData)
+TUniquePtr<FMLInferenceModel> FRuntimeORTCuda::CreateModel(TConstArrayView<uint8> ModelData)
 {
 	if (!CanCreateModel(ModelData))
 	{
-		return TSharedPtr<FMLInferenceModel>();
+		return TUniquePtr<FMLInferenceModel>();
 	}
 	const FMLInferenceNNXORTConf InConf;
 	FMLInferenceModelORTCuda* Model = new FMLInferenceModelORTCuda(&NNXEnvironmentORT, InConf);
 	if (!Model->Init(ModelData))
 	{
 		delete Model;
-		return TSharedPtr<FMLInferenceModel>();
+		return TUniquePtr<FMLInferenceModel>();
 	}
-	return TSharedPtr<FMLInferenceModel>(Model);
+	return TUniquePtr<FMLInferenceModel>(Model);
 }
 
-TSharedPtr<FMLInferenceModel> FRuntimeORTDml::CreateModel(TConstArrayView<uint8> ModelData)
+TUniquePtr<FMLInferenceModel> FRuntimeORTDml::CreateModel(TConstArrayView<uint8> ModelData)
 {
 	if (!CanCreateModel(ModelData))
 	{
-		return TSharedPtr<FMLInferenceModel>();
+		return TUniquePtr<FMLInferenceModel>();
 	}
 	const FMLInferenceNNXORTConf InConf;
 	FMLInferenceModelORTDml* Model = new FMLInferenceModelORTDml(&NNXEnvironmentORT, InConf);
 	if (!Model->Init(ModelData))
 	{
 		delete Model;
-		return TSharedPtr<FMLInferenceModel>();
+		return TUniquePtr<FMLInferenceModel>();
 	}
-	return TSharedPtr<FMLInferenceModel>(Model);
+	return TUniquePtr<FMLInferenceModel>(Model);
 }
 #endif
 
@@ -346,7 +346,7 @@ int FMLInferenceModelORT::SetInputTensorShapes(TConstArrayView<FTensorShape> InI
 	return 0;
 }
 
-int FMLInferenceModelORT::Run(TConstArrayView<FMLTensorBinding> InInputBindings, TConstArrayView<FMLTensorBinding> InOutputBindings)
+int FMLInferenceModelORT::RunSync(TConstArrayView<FMLTensorBinding> InInputBindings, TConstArrayView<FMLTensorBinding> InOutputBindings)
 {
 	DECLARE_SCOPE_CYCLE_COUNTER(TEXT("FMLInferenceModelORT_Run"), STAT_FMLInferenceModelORT_Run, STATGROUP_MachineLearning);
 
diff --git a/Engine/Plugins/Experimental/NNX/Source/NNXRuntimeORT/Private/NNXRuntimeORT.h b/Engine/Plugins/Experimental/NNX/Source/NNXRuntimeORT/Private/NNXRuntimeORT.h
--- a/Engine/Plugins/Experimental/NNX/Source/NNXRuntimeORT/Private/NNXRuntimeORT.h
+++ b/Engine/Plugins/Experimental/NNX/Source/NNXRuntimeORT/Private/NNXRuntimeORT.h
@@ -69,7 +69,7 @@ namespace NNX
 		virtual ~FRuntimeORTCpu() = default;
 		virtual FString GetRuntimeName() const override;
 		virtual EMLRuntimeSupportFlags GetSupportFlags() const override;
-		virtual TSharedPtr<FMLInferenceModel> CreateModel(TConstArrayView<uint8> ModelData);
+		virtual TUniquePtr<FMLInferenceModel> CreateModel(TConstArrayView<uint8> ModelData);
 	};
 
 #if PLATFORM_WINDOWS
@@ -81,7 +81,7 @@ namespace NNX
 		virtual ~FRuntimeORTCuda() = default;
 		virtual FString GetRuntimeName() const override;
 		virtual EMLRuntimeSupportFlags GetSupportFlags() const override;
-		virtual TSharedPtr<FMLInferenceModel> CreateModel(TConstArrayView<uint8> ModelData);
+		virtual TUniquePtr<FMLInferenceModel> CreateModel(TConstArrayView<uint8> ModelData);
 	};
 
 	class FRuntimeORTDml : public FRuntimeORT
@@ -92,7 +92,7 @@ namespace NNX
 		virtual ~FRuntimeORTDml() = default;
 		virtual FString GetRuntimeName() const override;
 		virtual EMLRuntimeSupportFlags GetSupportFlags() const override;
-		virtual TSharedPtr<FMLInferenceModel> CreateModel(TConstArrayView<uint8> ModelData);
+		virtual TUniquePtr<FMLInferenceModel> CreateModel(TConstArrayView<uint8> ModelData);
 	};
 #endif
 
@@ -189,7 +189,7 @@ namespace NNX
 		bool IsLoaded() const;
 
 		virtual int SetInputTensorShapes(TConstArrayView<FTensorShape> InInputShapes) override;
-		virtual int Run(TConstArrayView<FMLTensorBinding> InInputBindings, TConstArrayView<FMLTensorBinding> InOutputBindings) override;
+		virtual int RunSync(TConstArrayView<FMLTensorBinding> InInputBindings, TConstArrayView<FMLTensorBinding> InOutputBindings) override;
 
 		float GetLastRunTimeMSec() const;
 		UE::NNEProfiling::Internal::FStatistics GetRunStatistics() const;
diff --git a/Engine/Plugins/Experimental/NNX/Source/NNXRuntimeRDG/Private/NNXRuntimeDML.cpp b/Engine/Plugins/Experimental/NNX/Source/NNXRuntimeRDG/Private/NNXRuntimeDML.cpp
--- a/Engine/Plugins/Experimental/NNX/Source/NNXRuntimeRDG/Private/NNXRuntimeDML.cpp
+++ b/Engine/Plugins/Experimental/NNX/Source/NNXRuntimeRDG/Private/NNXRuntimeDML.cpp
@@ -1207,11 +1207,11 @@ public:
 		return EMLRuntimeSupportFlags::RDG;
 	}
 
-	virtual TSharedPtr<FMLInferenceModel> CreateModel(TConstArrayView<uint8> ModelData) override
+	virtual TUniquePtr<FMLInferenceModel> CreateModel(TConstArrayView<uint8> ModelData) override
 	{
 		if (!CanCreateModel(ModelData))
 		{
-			return TSharedPtr<FMLInferenceModel>();
+			return TUniquePtr<FMLInferenceModel>();
 		}
 
 		// Create the model and initialize it with the data not including the header
@@ -1219,9 +1219,9 @@ public:
 		if (!Model->Init(ModelData, &Ctx))
 		{
 			delete Model;
-			return TSharedPtr<FMLInferenceModel>();
+			return TUniquePtr<FMLInferenceModel>();
 		}
-		return TSharedPtr<FMLInferenceModel>(Model);
+		return TUniquePtr<FMLInferenceModel>(Model);
 	}
 
 	bool Init();
diff --git a/Engine/Plugins/Experimental/NNX/Source/NNXRuntimeRDG/Private/NNXRuntimeHLSL.cpp b/Engine/Plugins/Experimental/NNX/Source/NNXRuntimeRDG/Private/NNXRuntimeHLSL.cpp
--- a/Engine/Plugins/Experimental/NNX/Source/NNXRuntimeRDG/Private/NNXRuntimeHLSL.cpp
+++ b/Engine/Plugins/Experimental/NNX/Source/NNXRuntimeRDG/Private/NNXRuntimeHLSL.cpp
@@ -301,11 +301,11 @@ public:
 		return EMLRuntimeSupportFlags::RDG;
 	}
 
-	virtual TSharedPtr<FMLInferenceModel> CreateModel(TConstArrayView<uint8> ModelData) override
+	virtual TUniquePtr<FMLInferenceModel> CreateModel(TConstArrayView<uint8> ModelData) override
 	{
 		if (!CanCreateModel(ModelData))
 		{
-			return TSharedPtr<FMLInferenceModel>();
+			return TUniquePtr<FMLInferenceModel>();
 		}
 
 		// Create the model and initialize it with the data not including the header
@@ -313,9 +313,9 @@ public:
 		if (!Model->Init(ModelData))
 		{
 			delete Model;
-			return TSharedPtr<FMLInferenceModel>();
+			return TUniquePtr<FMLInferenceModel>();
 		}
-		return TSharedPtr<FMLInferenceModel>(Model);
+		return TUniquePtr<FMLInferenceModel>(Model);
 	}
 };
 
diff --git a/Engine/Plugins/Experimental/NNX/Source/NNXRuntimeRDG/Private/NNXRuntimeRDG.cpp b/Engine/Plugins/Experimental/NNX/Source/NNXRuntimeRDG/Private/NNXRuntimeRDG.cpp
--- a/Engine/Plugins/Experimental/NNX/Source/NNXRuntimeRDG/Private/NNXRuntimeRDG.cpp
+++ b/Engine/Plugins/Experimental/NNX/Source/NNXRuntimeRDG/Private/NNXRuntimeRDG.cpp
@@ -300,7 +300,7 @@ bool FMLInferenceModelRDG::LoadModel(TConstArrayView<uint8> ModelData, FMLRuntim
 /**
  * Run the inference model (synchronous version)
  */
-int FMLInferenceModelRDG::Run(TConstArrayView<FMLTensorBinding> InInputBindings, TConstArrayView<FMLTensorBinding> InOutputBindings)
+int FMLInferenceModelRDG::RunSync(TConstArrayView<FMLTensorBinding> InInputBindings, TConstArrayView<FMLTensorBinding> InOutputBindings)
 {
 	// Verify the model inputs were prepared
 	if (InputTensorShapes.Num() == 0)
diff --git a/Engine/Plugins/Experimental/NNX/Source/NNXRuntimeRDG/Private/NNXRuntimeRDG.h b/Engine/Plugins/Experimental/NNX/Source/NNXRuntimeRDG/Private/NNXRuntimeRDG.h
--- a/Engine/Plugins/Experimental/NNX/Source/NNXRuntimeRDG/Private/NNXRuntimeRDG.h
+++ b/Engine/Plugins/Experimental/NNX/Source/NNXRuntimeRDG/Private/NNXRuntimeRDG.h
@@ -92,7 +92,7 @@ public:
 	~FMLInferenceModelRDG();
 
 	virtual int SetInputTensorShapes(TConstArrayView<FTensorShape> InputShapes) override;
-	virtual int Run(TConstArrayView<FMLTensorBinding> InputBindings, TConstArrayView<FMLTensorBinding> OutputBindings) override;
+	virtual int RunSync(TConstArrayView<FMLTensorBinding> InputBindings, TConstArrayView<FMLTensorBinding> OutputBindings) override;
 	virtual int EnqueueRDG(FRDGBuilder& RDGBuilder, TConstArrayView<FMLTensorBinding> InputBindings, TConstArrayView<FMLTensorBinding> OutputBindings) override;
 
 protected:
