Implementation of HLSL InstanceNormalization operator.

diff --git a/Engine/Plugins/Experimental/NNX/Source/NNEHlslShaders/Internal/NNEHlslShadersInstanceNormalizationCS.h b/Engine/Plugins/Experimental/NNX/Source/NNEHlslShaders/Internal/NNEHlslShadersInstanceNormalizationCS.h
new file mode 100644
--- /dev/null
+++ b/Engine/Plugins/Experimental/NNX/Source/NNEHlslShaders/Internal/NNEHlslShadersInstanceNormalizationCS.h
@@ -0,0 +1,50 @@
+// Copyright Epic Games, Inc. All Rights Reserved.
+
+#pragma once
+
+#include "GlobalShader.h"
+#include "ShaderParameterUtils.h"
+#include "RenderGraphUtils.h"
+
+namespace NNX { class FTensor; }
+
+namespace UE::NNEHlslShaders::Internal
+{
+	enum class EInstanceNormalizationAlgorithm : uint8
+	{
+		Simple1x265 = 0,
+		SharedMemory8x32,
+		SharedMemory16x16,
+		SharedMemory32x8,
+		MAX
+	};
+
+	class NNEHLSLSHADERS_API TInstanceNormalizationCS : public FGlobalShader
+	{
+		DECLARE_GLOBAL_SHADER(TInstanceNormalizationCS);
+		SHADER_USE_PARAMETER_STRUCT(TInstanceNormalizationCS, FGlobalShader)
+
+		class FInstanceNormalizationAlgorithm : SHADER_PERMUTATION_ENUM_CLASS("ALGORITHM", EInstanceNormalizationAlgorithm);
+		using FPermutationDomain = TShaderPermutationDomain<FInstanceNormalizationAlgorithm>;
+
+	public:
+
+		BEGIN_SHADER_PARAMETER_STRUCT(FParameters, )
+			SHADER_PARAMETER(float, Epsilon)
+			SHADER_PARAMETER(int32, C)
+			SHADER_PARAMETER(int32, NxC)
+			SHADER_PARAMETER(int32, W)
+			SHADER_PARAMETER_RDG_BUFFER_SRV(Buffer<float>, Input)
+			SHADER_PARAMETER_RDG_BUFFER_SRV(Buffer<float>, Scale)
+			SHADER_PARAMETER_RDG_BUFFER_SRV(Buffer<float>, Bias)
+			SHADER_PARAMETER_RDG_BUFFER_UAV(RWBuffer<float>, Output)
+		END_SHADER_PARAMETER_STRUCT()
+
+		static void FillInParameters(float Epsilon, const NNX::FTensor& Input, FParameters& Parameters);
+
+		static FIntVector GetGroupCount(const FParameters& Parameters, EInstanceNormalizationAlgorithm Algorithm);
+		static EInstanceNormalizationAlgorithm GetAlgorithm(const FParameters& Parameters);
+
+		static void LexFromString(EInstanceNormalizationAlgorithm& OutValue, const TCHAR* StringVal);
+	};
+} // UE::NNEHlslShaders::Internal
\ No newline at end of file
diff --git a/Engine/Plugins/Experimental/NNX/Source/NNEHlslShaders/Private/NNEHlslShadersInstanceNormalizationCS.cpp b/Engine/Plugins/Experimental/NNX/Source/NNEHlslShaders/Private/NNEHlslShadersInstanceNormalizationCS.cpp
new file mode 100644
--- /dev/null
+++ b/Engine/Plugins/Experimental/NNX/Source/NNEHlslShaders/Private/NNEHlslShadersInstanceNormalizationCS.cpp
@@ -0,0 +1,59 @@
+// Copyright Epic Games, Inc. All Rights Reserved.
+
+#include "NNEHlslShadersInstanceNormalizationCS.h"
+#include "NNXCore.h"
+#include "NNXTypes.h"
+
+namespace UE::NNEHlslShaders::Internal
+{
+	namespace
+	{
+		FIntVector GetGroupSize(EInstanceNormalizationAlgorithm Algorithm)
+		{
+			switch (Algorithm)
+			{
+			case EInstanceNormalizationAlgorithm::Simple1x265: 			return { 1, 256, 1};
+			case EInstanceNormalizationAlgorithm::SharedMemory8x32:		return { 8,  32, 1};
+			case EInstanceNormalizationAlgorithm::SharedMemory16x16:	return {16,  16, 1};
+			case EInstanceNormalizationAlgorithm::SharedMemory32x8:		return {32,   8, 1};
+			}
+
+			check(false);
+			return {1, 256, 1};
+		}
+	}
+
+	void TInstanceNormalizationCS::FillInParameters(float Epsilon, const NNX::FTensor& Input, TInstanceNormalizationCS::FParameters& Parameters)
+	{
+		check(Input.GetShape().Rank() >= 3);
+		
+		Parameters.Epsilon = Epsilon;
+		const int32 N = Input.GetShape().Data[0];
+		Parameters.C = Input.GetShape().Data[1];
+		Parameters.NxC = N * Parameters.C;
+		Parameters.W = Input.GetShape().Volume() / (Parameters.NxC);
+	}
+
+	FIntVector TInstanceNormalizationCS::GetGroupCount(const TInstanceNormalizationCS::FParameters& Parameters, EInstanceNormalizationAlgorithm Algorithm)
+	{
+		int ThreadGroupCountValueY = FMath::DivideAndRoundUp((int32)Parameters.NxC, GetGroupSize(Algorithm).Y);
+
+		return {1, ThreadGroupCountValueY, 1};
+	}
+
+	EInstanceNormalizationAlgorithm TInstanceNormalizationCS::GetAlgorithm(const TInstanceNormalizationCS::FParameters& Parameters)
+	{
+		return EInstanceNormalizationAlgorithm::SharedMemory16x16;
+	}
+
+	void TInstanceNormalizationCS::LexFromString(EInstanceNormalizationAlgorithm& OutValue, const TCHAR* StringVal)
+	{
+		OutValue = EInstanceNormalizationAlgorithm::MAX;
+		if (FCString::Stricmp(StringVal, TEXT("Simple1x265")) == 0) OutValue = EInstanceNormalizationAlgorithm::Simple1x265;
+		else if (FCString::Stricmp(StringVal, TEXT("SharedMemory8x32")) == 0) OutValue = EInstanceNormalizationAlgorithm::SharedMemory8x32;
+		else if (FCString::Stricmp(StringVal, TEXT("SharedMemory16x16")) == 0) OutValue = EInstanceNormalizationAlgorithm::SharedMemory16x16;
+		else if (FCString::Stricmp(StringVal, TEXT("SharedMemory32x8")) == 0) OutValue = EInstanceNormalizationAlgorithm::SharedMemory32x8;
+	}
+
+	IMPLEMENT_GLOBAL_SHADER(TInstanceNormalizationCS, "/NNE/NNEHlslShadersInstanceNormalization.usf", "InstanceNormalization", SF_Compute);
+} // UE::NNEHlslShaders::Internal
\ No newline at end of file
diff --git a/Engine/Plugins/Experimental/NNX/Source/NNEHlslShaders/Shaders/NNEHlslShadersInstanceNormalization.usf b/Engine/Plugins/Experimental/NNX/Source/NNEHlslShaders/Shaders/NNEHlslShadersInstanceNormalization.usf
new file mode 100644
--- /dev/null
+++ b/Engine/Plugins/Experimental/NNX/Source/NNEHlslShaders/Shaders/NNEHlslShadersInstanceNormalization.usf
@@ -0,0 +1,176 @@
+// Copyright Epic Games, Inc. All Rights Reserved.
+
+#include "/Engine/Public/Platform.ush"
+
+#define WORK_TYPE float
+#define BUFFER_TYPE float
+#define READ(x) x
+#define WRITE(x) x
+
+int C;							// Number of channels
+int NxC;						// Precalculated N x C
+int W; 							// D1 x D2 ... Dn
+
+BUFFER_TYPE Epsilon;			// Attribute epsilon to use to avoid division by zero
+Buffer<BUFFER_TYPE> Input;		// Input data tensor with dimensions (N x C x D1 x D2 ... Dn)
+Buffer<BUFFER_TYPE> Scale;		// Scale tensor of size C
+Buffer<BUFFER_TYPE> Bias;		// Bias tensor of size C
+RWBuffer<BUFFER_TYPE> Output;	// Output tensor of the same shape as input
+
+#if ALGORITHM == 0
+#define GROUP_SIZE_X 1
+#define GROUP_SIZE_Y 256
+#elif ALGORITHM == 1
+#define GROUP_SIZE_X 8
+#define GROUP_SIZE_Y 32
+#elif ALGORITHM == 2
+#define GROUP_SIZE_X 16
+#define GROUP_SIZE_Y 16
+#elif ALGORITHM == 3
+#define GROUP_SIZE_X 32
+#define GROUP_SIZE_Y 8
+#endif
+
+#if GROUP_SIZE_X > 1
+
+groupshared WORK_TYPE SharedMemory[GROUP_SIZE_X * GROUP_SIZE_Y];
+
+[numthreads(GROUP_SIZE_X, GROUP_SIZE_Y, 1)]
+void InstanceNormalization(in const uint3 DispatchThreadID : SV_DispatchThreadID, in const uint3 GroupThreadID : SV_GroupThreadID)
+{
+	const int StartIdx = DispatchThreadID.y * W;
+	const int SharedMemoryOffsetY = GroupThreadID.y * GROUP_SIZE_X;
+
+	const WORK_TYPE InvNum = (WORK_TYPE)1 / (WORK_TYPE)W;
+
+	// Calculate mean
+	if (DispatchThreadID.y < NxC)
+	{
+		WORK_TYPE LocalSum = 0;
+		for (int CurrIdx = GroupThreadID.x; CurrIdx < W; CurrIdx += GROUP_SIZE_X)
+		{
+			LocalSum += READ(Input[StartIdx + CurrIdx]);
+		}
+
+		SharedMemory[SharedMemoryOffsetY + GroupThreadID.x] = LocalSum;
+	}
+
+	// Wait until all threads wrote result to shared memory
+	GroupMemoryBarrierWithGroupSync();
+
+	// Use one thread to calculate final result
+	if (GroupThreadID.x == 0 && DispatchThreadID.y < NxC)
+	{
+		WORK_TYPE LocalMean = 0;
+
+		UNROLL
+		for (int Idx = 0; Idx < GROUP_SIZE_X; Idx++)
+		{
+			LocalMean += SharedMemory[SharedMemoryOffsetY + Idx];
+		}
+
+		SharedMemory[SharedMemoryOffsetY] = LocalMean * InvNum;
+	}
+
+	// Wait until first thread wrote result to shared memory
+	GroupMemoryBarrierWithGroupSync();
+
+	WORK_TYPE Mean = 0;
+	if (DispatchThreadID.y < NxC)
+	{
+		Mean = SharedMemory[SharedMemoryOffsetY];
+	
+		// Calculate variance
+		WORK_TYPE LocalSum = 0;
+		WORK_TYPE Tmp = 0;
+		for (int CurrIdx = GroupThreadID.x; CurrIdx < W; CurrIdx += GROUP_SIZE_X)
+		{
+			Tmp = READ(Input[StartIdx + CurrIdx]) - Mean;
+
+			LocalSum += Tmp * Tmp;
+		}
+
+		SharedMemory[SharedMemoryOffsetY + GroupThreadID.x] = LocalSum;
+	}
+
+	// Wait until all threads wrote result to shared memory
+	GroupMemoryBarrierWithGroupSync();
+
+	// Use one thread to calculate final result
+	if (GroupThreadID.x == 0 && DispatchThreadID.y < NxC)
+	{
+		WORK_TYPE LocalVar = 0;
+
+		UNROLL
+		for (int Idx = 0; Idx < GROUP_SIZE_X; Idx++)
+		{
+			LocalVar += SharedMemory[SharedMemoryOffsetY + Idx];
+		}
+
+		SharedMemory[SharedMemoryOffsetY] = (WORK_TYPE)1 / sqrt(LocalVar * InvNum + Epsilon);
+	}
+
+	// Wait until first thread wrote result to shared memory
+	GroupMemoryBarrierWithGroupSync();
+
+	WORK_TYPE Variance = 0;
+	if (DispatchThreadID.y < NxC)
+	{
+		const int ChannelIdx = DispatchThreadID.y % C;
+		
+		WORK_TYPE InvStdev = SharedMemory[SharedMemoryOffsetY];
+
+		// Carry out instance normalization and write to output
+		for (int CurrIdx = GroupThreadID.x; CurrIdx < W; CurrIdx += GROUP_SIZE_X)
+		{
+			Output[StartIdx + CurrIdx] = WRITE((READ(Input[StartIdx + CurrIdx]) - Mean) * READ(Scale[ChannelIdx]) * InvStdev + READ(Bias[ChannelIdx]));
+		}
+	}
+}
+
+#else
+
+[numthreads(1, 256, 1)]
+void InstanceNormalization(in const uint3 DispatchThreadID : SV_DispatchThreadID)
+{
+	if (DispatchThreadID.y >= NxC || DispatchThreadID.x >= 1)
+	{
+		return;
+	}
+
+	const int StartIdx = DispatchThreadID.y * W;
+	const int ChannelIdx = DispatchThreadID.y % C;
+
+	const WORK_TYPE InvNum = (WORK_TYPE)1 / (WORK_TYPE)W;
+
+	// Calculate mean
+	WORK_TYPE Mean = 0;
+	for (int Idx = 0; Idx < W; Idx++)
+	{
+		Mean += READ(Input[StartIdx + Idx]);
+	}
+
+	Mean *= InvNum;
+
+	// Calculate variance
+	WORK_TYPE Variance = 0;
+	WORK_TYPE Tmp = 0;
+	for (int Idx = 0; Idx < W; Idx++)
+	{
+		Tmp = READ(Input[StartIdx + Idx]) - Mean;
+
+		Variance += Tmp * Tmp;
+	}
+
+	Variance *= InvNum;
+
+	WORK_TYPE InvStdev = (WORK_TYPE)1 / sqrt(Variance + Epsilon);
+
+	// Carry out instance normalization and write to output
+	for (int Idx = 0; Idx < W; Idx++)
+	{
+		Output[StartIdx + Idx] = WRITE((READ(Input[StartIdx + Idx]) - Mean) * READ(Scale[ChannelIdx]) * InvStdev + READ(Bias[ChannelIdx]));
+	}
+}
+
+#endif
\ No newline at end of file
diff --git a/Engine/Plugins/Experimental/NNX/Source/NNXQA/Resources/NNXQATestDesc.json b/Engine/Plugins/Experimental/NNX/Source/NNXQA/Resources/NNXQATestDesc.json
--- a/Engine/Plugins/Experimental/NNX/Source/NNXQA/Resources/NNXQATestDesc.json
+++ b/Engine/Plugins/Experimental/NNX/Source/NNXQA/Resources/NNXQATestDesc.json
@@ -302,6 +302,15 @@
 			"targets": [
 				{ "target": "Pad" }
 			]
+		},
+		{
+			"category": "InstanceNormalization",
+			"runtimes": [
+				{ "name": "NNXRuntimeDml", "skip": true }
+			],
+			"targets": [
+				{ "target": "InstanceNormalization" }
+			]
 		}
 	],
 
@@ -583,6 +592,82 @@
 				  "weights": [ { "shape": [ 2 ], "source": [ "1", "2" ] , "type": "int64" } ],
 				  "outputs": [ { "shape": [ 4 ], "source": [ "0", "1", "0", "0" ] } ] }
 			]
+		},
+		{
+			"name": "InstanceNormalization",
+			"datasets": [
+				{ "inputs": [ { "shape": [ 1, 1, 3 ] }, { "shape": [ 1 ] }, { "shape": [ 1 ] } ], "outputs": [ { "shape": [ 1, 1, 3 ] } ], "runtimes": [ { "name": "NNXRuntimeORTDml", "skip": true, "comment": "Not supported" } ] },
+				{ "inputs": [ { "shape": [ 1, 2, 3 ] }, { "shape": [ 2 ] }, { "shape": [ 2 ] } ], "outputs": [ { "shape": [ 1, 2, 3 ] } ], "runtimes": [ { "name": "NNXRuntimeORTDml", "skip": true, "comment": "Not supported" } ] },
+				{ "inputs": [ { "shape": [ 2, 2, 3 ] }, { "shape": [ 2 ] }, { "shape": [ 2 ] } ], "outputs": [ { "shape": [ 2, 2, 3 ] } ], "runtimes": [ { "name": "NNXRuntimeORTDml", "skip": true, "comment": "Not supported" } ]},
+				{ "inputs": [ { "shape": [ 1, 2, 1, 3 ] }, { "shape": [ 2 ] }, { "shape": [ 2 ] } ], "outputs": [ { "shape": [ 1, 2, 1, 3 ] } ] },
+				{ "inputs": [ { "shape": [ 2, 2, 1, 3 ] }, { "shape": [ 2 ] }, { "shape": [ 2 ] } ], "outputs": [ { "shape": [ 2, 2, 1, 3 ] } ] },
+				{ "inputs": [ { "shape": [ 2, 2, 2, 2 ] }, { "shape": [ 2 ] }, { "shape": [ 2 ] } ], "outputs": [ { "shape": [ 2, 2, 2, 2 ] } ] },
+				{
+					"inputs": [ { "shape": [ 2, 2, 4, 1, 3 ] }, { "shape": [ 2 ] }, { "shape": [ 2 ] } ], "outputs": [ { "shape": [ 2, 2, 4, 1, 3 ] } ],
+					"runtimes": [
+						{ "name": "NNXRuntimeHLSL", "relative_error": 3.5e-3 },
+						{ "name": "NNXRuntimeORTCpu", "relative_error": 3.5e-3 },
+						{ "name": "NNXRuntimeORTDml", "skip": true, "comment": "Not supported" }
+					]
+				},
+				{
+					"inputs": [ { "shape": [ 2, 2, 4, 4 ] }, { "shape": [ 2 ] }, { "shape": [ 2 ] } ], "outputs": [ { "shape": [ 2, 2, 4, 4 ] } ],
+					"runtimes": [
+						{ "name": "NNXRuntimeHLSL", "relative_error": 3.5e-3 },
+						{ "name": "NNXRuntimeORTCpu", "relative_error": 3.5e-3 },
+						{ "name": "NNXRuntimeORTDml", "relative_error": 3.5e-3 }
+					]
+				},
+				{
+					"inputs": [ { "shape": [ 2, 2, 8, 8 ] }, { "shape": [ 2 ] }, { "shape": [ 2 ] } ], "outputs": [ { "shape": [ 2, 2, 8, 8 ] } ],
+					"runtimes": [
+						{ "name": "NNXRuntimeHLSL", "relative_error": 2 },
+						{ "name": "NNXRuntimeORTCpu", "relative_error": 2 },
+						{ "name": "NNXRuntimeORTDml", "relative_error": 2 }
+					]
+				},
+				{
+					"inputs": [ { "shape": [ 2, 128, 8, 8 ] }, { "shape": [ 128 ] }, { "shape": [ 128 ] } ], "outputs": [ { "shape": [ 2, 128, 8, 8 ] } ],
+					"runtimes": [
+						{ "name": "NNXRuntimeHLSL", "relative_error": 2 },
+						{ "name": "NNXRuntimeORTCpu", "relative_error": 2 },
+						{ "name": "NNXRuntimeORTDml", "relative_error": 2 }
+					]
+				},
+				{
+					"inputs": [ { "shape": [ 128, 2, 8, 8 ] }, { "shape": [ 2 ] }, { "shape": [ 2 ] } ], "outputs": [ { "shape": [ 128, 2, 8, 8 ] } ],
+					"runtimes": [
+						{ "name": "NNXRuntimeHLSL", "relative_error": 2 },
+						{ "name": "NNXRuntimeORTCpu", "relative_error": 2 },
+						{ "name": "NNXRuntimeORTDml", "relative_error": 2 }
+					]
+				},
+				{
+					"inputs": [ { "shape": [ 64, 64, 8, 8 ] }, { "shape": [ 64 ] }, { "shape": [ 64 ] } ], "outputs": [ { "shape": [ 64, 64, 8, 8 ] } ],
+					"runtimes": [
+						{ "name": "NNXRuntimeHLSL", "relative_error": 2 },
+						{ "name": "NNXRuntimeORTCpu", "relative_error": 2 },
+						{ "name": "NNXRuntimeORTDml", "relative_error": 2 }
+					]
+				},
+				{
+					"inputs": [ { "shape": [ 2, 2, 45, 45 ] }, { "shape": [ 2 ] }, { "shape": [ 2 ] } ], "outputs": [ { "shape": [ 2, 2, 45, 45 ] } ],
+					"runtimes": [
+						{ "name": "NNXRuntimeHLSL", "relative_error": 5 },
+						{ "name": "NNXRuntimeORTCpu", "relative_error": 5 },
+						{ "name": "NNXRuntimeORTDml", "relative_error": 5 }
+					]
+				},
+				{
+					"inputs": [ { "shape": [ 1, 128, 56, 56 ] }, { "shape": [ 128 ] }, { "shape": [ 128 ] } ], "outputs": [ { "shape": [ 1, 128, 56, 56 ] } ],
+					"runtimes": [
+						{ "name": "NNXRuntimeHLSL", "relative_error": 9 },
+						{ "name": "NNXRuntimeORTCpu", "relative_error": 9 },
+						{ "name": "NNXRuntimeORTDml", "relative_error": 9 }
+					],
+					"comment": "from 'mosaic-9-sim.onnx'"
+				}
+			]
 		}
 	],
 
@@ -834,6 +919,15 @@
 					{ "name": "strides", "value": [ 2, 2 ], "type": "Int32Array"}
 				] }
 			]
+		},
+		{
+			"name": "InstanceNormalization.Algorithms",
+			"attribute_maps": [
+				{ "attributes": [ { "name": "__UE__algorithm", "value": "Simple1x265", "type": "String" } ] },
+				{ "attributes": [ { "name": "__UE__algorithm", "value": "SharedMemory8x32", "type": "String" } ] },
+				{ "attributes": [ { "name": "__UE__algorithm", "value": "SharedMemory16x16", "type": "String" } ] },
+				{ "attributes": [ { "name": "__UE__algorithm", "value": "SharedMemory32x8", "type": "String" } ] }
+			]
 		}
 	]
 	
diff --git a/Engine/Plugins/Experimental/NNX/Source/NNXRuntimeRDG/Private/Hlsl/NNIRuntimeRDGInstanceNormalization.cpp b/Engine/Plugins/Experimental/NNX/Source/NNXRuntimeRDG/Private/Hlsl/NNIRuntimeRDGInstanceNormalization.cpp
new file mode 100644
--- /dev/null
+++ b/Engine/Plugins/Experimental/NNX/Source/NNXRuntimeRDG/Private/Hlsl/NNIRuntimeRDGInstanceNormalization.cpp
@@ -0,0 +1,173 @@
+// Copyright Epic Games, Inc. All Rights Reserved.
+
+#include "NNIRuntimeRDGInstanceNormalization.h"
+#include "NNEHlslShadersInstanceNormalizationCS.h"
+#include "NNXRuntimeHLSLHelper.h"
+
+namespace
+{
+
+int ValidateInput(const NNX::FTensorShape& Input, const NNX::FTensorShape& Scale, const NNX::FTensorShape& Bias)
+{
+	if (Input.Rank() < 3)
+	{
+		UE_LOG(LogNNX, Warning, TEXT("InstanceNormalization input data should be at least of rank 4: %d"), Input.Rank());
+		return -1;
+	}
+
+	if (Scale.Rank() != 1)
+	{
+		UE_LOG(LogNNX, Warning, TEXT("InstanceNormalization input scale should be of rank 1: %d"), Scale.Rank());
+		return -1;
+	}
+	if (Scale.Data[0] != Input.Data[1])
+	{
+		UE_LOG(LogNNX, Warning, TEXT("InstanceNormalization input scale size should be equal to channel count: %d vs %d"), Scale.Data[0], Input.Data[1]);
+		return -1;
+	}
+
+	if (Bias.Data[0] != Input.Data[1])
+	{
+		UE_LOG(LogNNX, Warning, TEXT("InstanceNormalization intput B size should be equal to channel count: : %d vs %d"), Bias.Data[0], Input.Data[1]);
+		return -1;
+	}
+	if (Bias.Rank() != 1)
+	{
+		UE_LOG(LogNNX, Warning, TEXT("InstanceNormalization input B should be of rank 1: %d"), Bias.Rank());
+		return -1;
+	}
+
+	return 0;
+}
+
+}
+
+namespace UE::NNIRuntimeRDG::Private::Hlsl
+{
+	DECLARE_GPU_STAT_NAMED(FNNIOperatorInstanceNormalization, TEXT("NNI.Operator.Hlsl.InstanceNormalization"));
+
+	using EInstanceNormalizationAlgorithm = UE::NNEHlslShaders::Internal::EInstanceNormalizationAlgorithm;
+
+	/**
+	 * InstanceNormalization operator implementation
+	 */
+	class TInstanceNormalization : public NNX::FMLOperatorHlsl
+	{
+	public:
+
+		TInstanceNormalization() {}
+		virtual ~TInstanceNormalization() = default;
+
+	private:
+
+		float Epsilon = 1e-5;
+
+		EInstanceNormalizationAlgorithm Algorithm = EInstanceNormalizationAlgorithm::MAX;
+
+	public:
+
+		virtual int ComputeOutputShape(TConstArrayView<NNX::FTensorShape> InputShapes, TArray<NNX::FTensorShape>& OutputShapes) const override
+		{
+			check(InputShapes.Num() == 3);
+
+			OutputShapes.Empty(); // TODO REMOVE caller should make sure empty array is passed!
+
+			if (const int Res = ValidateInput(InputShapes[0], InputShapes[1], InputShapes[2]); Res < 0)
+			{
+				return Res;
+			}
+
+			OutputShapes.Add(InputShapes[0]);
+
+			return 0;
+		};
+
+		virtual bool Initialize(TConstArrayView<NNX::FTensorDesc> InputTensorDescs, TConstArrayView<NNX::FTensorDesc> OutputTensorDescs, const UE::NNECore::FAttributeMap& Attributes) override
+		{
+			using namespace UE::NNEHlslShaders::Internal;
+
+			check(InputTensorDescs.Num() == 3);
+			check(OutputTensorDescs.Num() == 1);
+
+			Epsilon = Attributes.GetValue<float>(TEXT("epsilon"));
+
+			// For testing only
+			TInstanceNormalizationCS::LexFromString(Algorithm, *Attributes.GetValueOrDefault<FString>(TEXT("__UE__algorithm"), "MAX"));
+
+			return true;
+		}
+
+		virtual void Dispatch(FRDGBuilder& GraphBuilder, TConstArrayView<NNX::FTensorRDG> InputTensors, TConstArrayView<NNX::FTensorRDG> OutputTensors) override
+		{
+			using namespace UE::NNEHlslShaders::Internal;
+
+			check(InputTensors.Num() == 3);
+			check(OutputTensors.Num() == 1);
+
+			const NNX::FTensorRDG& Input = InputTensors[0];
+			const NNX::FTensorRDG& Scale = InputTensors[1];
+			const NNX::FTensorRDG& Bias = InputTensors[2];
+			const NNX::FTensorRDG& Output = OutputTensors[0];
+
+			// Set parameters
+			TInstanceNormalizationCS::FParameters* Parameters = GraphBuilder.AllocParameters<TInstanceNormalizationCS::FParameters>();
+			TInstanceNormalizationCS::FillInParameters(Epsilon, Input, *Parameters);
+			Parameters->Input = GraphBuilder.CreateSRV(FRDGBufferSRVDesc(Input.GetBuffer(), PF_R32_FLOAT));
+			Parameters->Scale = GraphBuilder.CreateSRV(FRDGBufferSRVDesc(Scale.GetBuffer(), PF_R32_FLOAT));
+			Parameters->Bias = GraphBuilder.CreateSRV(FRDGBufferSRVDesc(Bias.GetBuffer(), PF_R32_FLOAT));
+			Parameters->Output = GraphBuilder.CreateUAV(FRDGBufferUAVDesc(Output.GetBuffer(), PF_R32_FLOAT));
+
+			TInstanceNormalizationCS::FPermutationDomain PermutationVector;
+			PermutationVector.Set<TInstanceNormalizationCS::FInstanceNormalizationAlgorithm>(Algorithm);
+			TShaderMapRef<TInstanceNormalizationCS> ComputeShader(GetGlobalShaderMap(GMaxRHIFeatureLevel), PermutationVector);
+
+			if (Algorithm == EInstanceNormalizationAlgorithm::MAX)
+			{
+				TInstanceNormalizationCS::GetAlgorithm(*Parameters);
+			}
+
+			FIntVector ThreadGroupCount = TInstanceNormalizationCS::GetGroupCount(*Parameters, Algorithm);
+
+			RDG_EVENT_SCOPE(GraphBuilder, "NNI.Operator.Hlsl.InstanceNormalization");
+			RDG_GPU_STAT_SCOPE(GraphBuilder, FNNIOperatorInstanceNormalization);
+
+			FComputeShaderUtils::AddPass(
+				GraphBuilder,
+				RDG_EVENT_NAME("NNI.Operator.Hlsl.InstanceNormalization.Dispatch"),
+				ERDGPassFlags::Compute | ERDGPassFlags::NeverCull,
+				ComputeShader,
+				Parameters,
+				ThreadGroupCount);
+		}
+	};
+
+	bool ValidateInstanceNormalizationOperator(const UE::NNECore::FAttributeMap& AttributeMap, TConstArrayView<EMLTensorDataType> InputTypes, TConstArrayView<NNX::FSymbolicTensorShape> InputShapes)
+	{
+		bool bIsValid = true;
+
+		NNX::FAttributeValidator AttributeValidator;
+		AttributeValidator.AddRequired(TEXT("epsilon"), ENNEAttributeDataType::Float);
+		AttributeValidator.AddOptional(TEXT("__UE__algorithm"), ENNEAttributeDataType::String);
+		bIsValid &= AttributeValidator.Validate(AttributeMap);
+
+		NNX::FInputValidator InputValidator;
+		InputValidator.AddSupportedType(EMLTensorDataType::Float);
+		InputValidator.AddRequired();
+		InputValidator.AddRequired();
+		InputValidator.AddRequired();
+		bIsValid &= InputValidator.Validate(InputTypes);
+
+		return bIsValid;
+	}
+
+	NNX::FMLOperatorHlsl* CreateInstanceNormalizationOperator()
+	{
+		return new TInstanceNormalization();
+	}
+
+	bool RegisterInstanceNormalizationOperator(NNX::FMLOperatorRegistryHlsl& Registry)
+	{
+		Registry.OpAdd(TEXT("InstanceNormalization"), CreateInstanceNormalizationOperator, ValidateInstanceNormalizationOperator);
+		return true;
+	}
+} // UE::NNIRuntimeRDG::Private::Hlsl
diff --git a/Engine/Plugins/Experimental/NNX/Source/NNXRuntimeRDG/Private/Hlsl/NNIRuntimeRDGInstanceNormalization.h b/Engine/Plugins/Experimental/NNX/Source/NNXRuntimeRDG/Private/Hlsl/NNIRuntimeRDGInstanceNormalization.h
new file mode 100644
--- /dev/null
+++ b/Engine/Plugins/Experimental/NNX/Source/NNXRuntimeRDG/Private/Hlsl/NNIRuntimeRDGInstanceNormalization.h
@@ -0,0 +1,11 @@
+// Copyright Epic Games, Inc. All Rights Reserved.
+
+#pragma once
+
+#include "../NNXRuntimeRDG.h"
+#include "../NNXRuntimeHLSLOp.h"
+
+namespace UE::NNIRuntimeRDG::Private::Hlsl
+{
+	bool RegisterInstanceNormalizationOperator(NNX::FMLOperatorRegistryHlsl& Registry);
+} // UE::NNIRuntimeRDG::Private::Hlsl
diff --git a/Engine/Plugins/Experimental/NNX/Source/NNXRuntimeRDG/Private/NNXRuntimeDML.cpp b/Engine/Plugins/Experimental/NNX/Source/NNXRuntimeRDG/Private/NNXRuntimeDML.cpp
--- a/Engine/Plugins/Experimental/NNX/Source/NNXRuntimeRDG/Private/NNXRuntimeDML.cpp
+++ b/Engine/Plugins/Experimental/NNX/Source/NNXRuntimeRDG/Private/NNXRuntimeDML.cpp
@@ -1207,6 +1207,28 @@ public:
 		return EMLRuntimeSupportFlags::RDG;
 	}
 
+	virtual TArray<uint8> CreateModelData(FString FileType, TConstArrayView<uint8> FileData) override
+	{
+		if (!CanCreateModelData(FileType, FileData))
+		{
+			return {};
+		}
+
+		TUniquePtr<IModelOptimizer> Optimizer = CreateONNXToNNXModelOptimizer();
+
+		FNNIModelRaw InputModel;
+		InputModel.Data = FileData;
+		InputModel.Format = ENNXInferenceFormat::ONNX;
+
+		FNNIModelRaw OutputModel;
+		if (!Optimizer->Optimize(InputModel, OutputModel, {}))
+		{
+			return {};
+		}
+
+		return ConvertToModelData(OutputModel.Data);
+	};
+
 	virtual TUniquePtr<FMLInferenceModel> CreateModel(TConstArrayView<uint8> ModelData) override
 	{
 		if (!CanCreateModel(ModelData))
diff --git a/Engine/Plugins/Experimental/NNX/Source/NNXRuntimeRDG/Private/NNXRuntimeHLSL.cpp b/Engine/Plugins/Experimental/NNX/Source/NNXRuntimeRDG/Private/NNXRuntimeHLSL.cpp
--- a/Engine/Plugins/Experimental/NNX/Source/NNXRuntimeRDG/Private/NNXRuntimeHLSL.cpp
+++ b/Engine/Plugins/Experimental/NNX/Source/NNXRuntimeRDG/Private/NNXRuntimeHLSL.cpp
@@ -12,6 +12,7 @@
 #include "Hlsl/NNIRuntimeRDGConvTranspose.h"
 #include "Hlsl/NNIRuntimeRDGGemm.h"
 #include "Hlsl/NNIRuntimeRDGMatMul.h"
+#include "Hlsl/NNIRuntimeRDGInstanceNormalization.h"
 #include "NNXRuntimeHLSLHelper.h"
 #include "NNXRuntimeRDG.h"
 
@@ -286,6 +287,7 @@ public:
 		UE::NNIRuntimeRDG::Private::Hlsl::RegisterConvOperator(*registry);
 		UE::NNIRuntimeRDG::Private::Hlsl::RegisterConvTransposeOperator(*registry);
 		UE::NNIRuntimeRDG::Private::Hlsl::RegisterMatMulOperator(*registry);
+		UE::NNIRuntimeRDG::Private::Hlsl::RegisterInstanceNormalizationOperator(*registry);
 
 		return true;
 	}
@@ -301,6 +303,29 @@ public:
 		return EMLRuntimeSupportFlags::RDG;
 	}
 
+	virtual TArray<uint8> CreateModelData(FString FileType, TConstArrayView<uint8> FileData) override
+	{
+		if (!CanCreateModelData(FileType, FileData))
+		{
+			return {};
+		}
+
+		TUniquePtr<IModelOptimizer> Optimizer = CreateONNXToNNXModelOptimizer();
+		Optimizer->AddValidator(MakeShared<FModelValidatorHlsl>());
+
+		FNNIModelRaw InputModel;
+		InputModel.Data = FileData;
+		InputModel.Format = ENNXInferenceFormat::ONNX;
+
+		FNNIModelRaw OutputModel;
+		if (!Optimizer->Optimize(InputModel, OutputModel, {}))
+		{
+			return {};
+		}
+
+		return ConvertToModelData(OutputModel.Data);
+	};
+
 	virtual TUniquePtr<FMLInferenceModel> CreateModel(TConstArrayView<uint8> ModelData) override
 	{
 		if (!CanCreateModel(ModelData))
diff --git a/Engine/Plugins/Experimental/NNX/Source/NNXRuntimeRDG/Private/NNXRuntimeRDG.cpp b/Engine/Plugins/Experimental/NNX/Source/NNXRuntimeRDG/Private/NNXRuntimeRDG.cpp
--- a/Engine/Plugins/Experimental/NNX/Source/NNXRuntimeRDG/Private/NNXRuntimeRDG.cpp
+++ b/Engine/Plugins/Experimental/NNX/Source/NNXRuntimeRDG/Private/NNXRuntimeRDG.cpp
@@ -179,35 +179,6 @@ bool FMLRuntimeRDG::CanCreateModelData(FString FileType, TConstArrayView<uint8>
 	return FileType.Compare("onnx", ESearchCase::IgnoreCase) == 0;
 }
 
-TArray<uint8> FMLRuntimeRDG::CreateModelData(FString FileType, TConstArrayView<uint8> FileData)
-{
-	if (!CanCreateModelData(FileType, FileData))
-	{
-		return {};
-	}
-
-	TUniquePtr<IModelOptimizer> Optimizer = CreateONNXToNNXModelOptimizer();
-
-	FNNIModelRaw InputModel;
-	InputModel.Data = FileData;
-	InputModel.Format = ENNXInferenceFormat::ONNX;
-	FNNIModelRaw OutputModel;
-	FOptimizerOptionsMap Options;
-	if (!Optimizer->Optimize(InputModel, OutputModel, Options))
-	{
-		return {};
-	}
-
-	int32 GuidSize = sizeof(FMLRuntimeRDG::GUID);
-	int32 VersionSize = sizeof(FMLRuntimeRDG::Version);
-	TArray<uint8> Result;
-	FMemoryWriter Writer(Result);
-	Writer << FMLRuntimeRDG::GUID;
-	Writer << FMLRuntimeRDG::Version;
-	Writer.Serialize(OutputModel.Data.GetData(), OutputModel.Data.Num());
-	return Result;
-};
-
 bool FMLRuntimeRDG::CanCreateModel(TConstArrayView<uint8> ModelData) const
 {
 	int32 GuidSize = sizeof(FMLRuntimeRDG::GUID);
@@ -243,9 +214,9 @@ bool FMLInferenceModelRDG::LoadModel(TConstArrayView<uint8> ModelData, FMLRuntim
 {
 	int32 GuidSize = sizeof(FMLRuntimeRDG::GUID);
 	int32 VersionSize = sizeof(FMLRuntimeRDG::Version);
-	TArray<uint8> ModelBuffer = TArray<uint8>(&(ModelData.GetData()[GuidSize + VersionSize]), ModelData.Num() - GuidSize - VersionSize);
+	TConstArrayView<uint8> ModelBuffer = {&(ModelData.GetData()[GuidSize + VersionSize]), ModelData.Num() - GuidSize - VersionSize};
 
-	FMemoryReader Reader(ModelBuffer);
+	FMemoryReaderView Reader(ModelBuffer);
 
 	FMLRuntimeFormat::StaticStruct()->SerializeBin(Reader, &Format);
 
@@ -592,4 +563,16 @@ void FMLInferenceModelRDG::AddTensorReadbacks_RenderThread(FRDGBuilder& GraphBui
 	}
 }
 
+TArray<uint8> ConvertToModelData(TArrayView<uint8> ModelBuffer)
+{
+	TArray<uint8> Result;
+
+	FMemoryWriter Writer(Result);
+	Writer << FMLRuntimeRDG::GUID;
+	Writer << FMLRuntimeRDG::Version;
+	Writer.Serialize(ModelBuffer.GetData(), ModelBuffer.Num());
+
+	return Result;
+}
+
 } // NNX
diff --git a/Engine/Plugins/Experimental/NNX/Source/NNXRuntimeRDG/Private/NNXRuntimeRDG.h b/Engine/Plugins/Experimental/NNX/Source/NNXRuntimeRDG/Private/NNXRuntimeRDG.h
--- a/Engine/Plugins/Experimental/NNX/Source/NNXRuntimeRDG/Private/NNXRuntimeRDG.h
+++ b/Engine/Plugins/Experimental/NNX/Source/NNXRuntimeRDG/Private/NNXRuntimeRDG.h
@@ -75,7 +75,7 @@ public:
 	static int32 Version;
 
 	virtual bool CanCreateModelData(FString FileType, TConstArrayView<uint8> FileData) const;
-	virtual TArray<uint8> CreateModelData(FString FileType, TConstArrayView<uint8> FileData);
+	virtual TArray<uint8> CreateModelData(FString FileType, TConstArrayView<uint8> FileData) = 0;
 	virtual bool CanCreateModel(TConstArrayView<uint8> ModelData) const;
 };
 
@@ -303,6 +303,8 @@ public:
 	}
 };
 
+TArray<uint8> ConvertToModelData(TArrayView<uint8> ModelBuffer);
+
 // NOTE: For now we only have DML on Windows, we should add support for XSX
 #ifdef NNE_USE_DIRECTML
 
@@ -314,7 +316,6 @@ extern void FMLRuntimeDmlShutdown();
 extern IRuntime* FMLRuntimeHlslStartup();
 extern void FMLRuntimeHlslShutdown();
 
-
 ///// OperatorId
 //class FOperatorId
 //{
